Added the experimental Localization Dashboard UI.

diff --git a/Engine/Source/Editor/EditorStyle/Private/SlateEditorStyle.cpp b/Engine/Source/Editor/EditorStyle/Private/SlateEditorStyle.cpp
--- a/Engine/Source/Editor/EditorStyle/Private/SlateEditorStyle.cpp
+++ b/Engine/Source/Editor/EditorStyle/Private/SlateEditorStyle.cpp
@@ -145,6 +145,7 @@ void FSlateEditorStyle::FStyle::Initialize()
 	SetupLandscapeEditorStyle();
 	SetupToolkitStyles();
 	SetupTranslationEditorStyles();
+	SetupLocalizationDashboardStyles();
 	SetupMatineeStyle();
 	SetupSourceControlStyles();
 	SetupAutomationStyles();
@@ -6278,6 +6279,36 @@ void FSlateEditorStyle::FStyle::SetupTranslationEditorStyles()
 	Set("TranslationEditor.TranslationPicker", new IMAGE_BRUSH("Icons/icon_StaticMeshEd_VertColor_40x", Icon40x40));
 }
 
+void FSlateEditorStyle::FStyle::SetupLocalizationDashboardStyles()
+{
+	Set("LocalizationDashboard.GatherAllTargets", new IMAGE_BRUSH("Icons/icon_Refresh_16x", Icon40x40));
+	Set("LocalizationDashboard.ImportForAllTargetsCultures", new IMAGE_BRUSH("Icons/icon_Samples_Download_40x", Icon40x40));
+	Set("LocalizationDashboard.ExportForAllTargetsCultures", new IMAGE_BRUSH("Icons/icon_Import_40x", Icon40x40));
+
+	Set("LocalizationDashboard.GatherTarget", new IMAGE_BRUSH("Icons/icon_Refresh_16x", Icon16x16));
+	Set("LocalizationDashboard.RefreshWordCount", new IMAGE_BRUSH("Icons/icon_Refresh_16x", Icon16x16));
+	Set("LocalizationDashboard.ImportForAllCultures", new IMAGE_BRUSH("Icons/icon_Samples_Download_20x", Icon16x16));
+	Set("LocalizationDashboard.ExportForAllCultures", new IMAGE_BRUSH("Icons/icon_Import_16x", Icon16x16));
+	Set("LocalizationDashboard.DeleteTarget", new IMAGE_BRUSH("Icons/Cross_12x", Icon12x12 ) );
+
+	Set("LocalizationTargetEditor.Gather", new IMAGE_BRUSH("Icons/icon_Refresh_16x", Icon40x40));
+	Set("LocalizationTargetEditor.ImportForAllCultures", new IMAGE_BRUSH("Icons/icon_Samples_Download_40x", Icon40x40));
+	Set("LocalizationTargetEditor.ExportForAllCultures", new IMAGE_BRUSH("Icons/icon_Import_40x", Icon40x40));
+	Set("LocalizationTargetEditor.RefreshWordCounts", new IMAGE_BRUSH("Icons/icon_Refresh_16x", Icon40x40));
+
+	Set("LocalizationTargetEditor.NativeCulture", new IMAGE_BRUSH( "Icons/Star_16x", Icon16x16 ) );
+	Set("LocalizationTargetEditor.EditTranslations", new IMAGE_BRUSH("Icons/icon_file_open_16px", Icon16x16));
+	Set("LocalizationTargetEditor.ImportForCulture", new IMAGE_BRUSH("Icons/icon_Samples_Download_20x", Icon16x16));
+	Set("LocalizationTargetEditor.ExportForCulture", new IMAGE_BRUSH("Icons/icon_Import_16x", Icon16x16));
+	Set("LocalizationTargetEditor.DeleteCulture", new IMAGE_BRUSH("Icons/Cross_12x", Icon12x12 ) );
+
+	Set( "LocalizationDashboard.CommandletLog.ListBorder", new BOX_BRUSH( "/Docking/AppTabContentArea", FMargin(4/16.0f) ) );
+	Set( "LocalizationDashboard.CommandletLog.Text", FTextBlockStyle(NormalText)
+		.SetFont( TTF_CORE_FONT( "Fonts/Roboto-Regular", 8 ) )
+		.SetShadowOffset( FVector2D::ZeroVector )
+		);
+}
+
 #undef IMAGE_BRUSH
 #undef BOX_BRUSH
 #undef BORDER_BRUSH
diff --git a/Engine/Source/Editor/EditorStyle/Private/SlateEditorStyle.h b/Engine/Source/Editor/EditorStyle/Private/SlateEditorStyle.h
--- a/Engine/Source/Editor/EditorStyle/Private/SlateEditorStyle.h
+++ b/Engine/Source/Editor/EditorStyle/Private/SlateEditorStyle.h
@@ -67,6 +67,7 @@ public:
 		void SetupDockingStyles();
 		void SetupTutorialStyles();
 		void SetupTranslationEditorStyles();
+		void SetupLocalizationDashboardStyles();
 		void SetupPropertyEditorStyles();
 		void SetupProfilerStyle();
 		void SetupGraphEditorStyles();
diff --git a/Engine/Source/Editor/LocalizationDashboard/Classes/LocalizationTarget.h b/Engine/Source/Editor/LocalizationDashboard/Classes/LocalizationTarget.h
new file mode 100644
--- /dev/null
+++ b/Engine/Source/Editor/LocalizationDashboard/Classes/LocalizationTarget.h
@@ -0,0 +1,78 @@
+// Copyright 1998-2014 Epic Games, Inc. All Rights Reserved.
+
+#pragma once
+
+#include "Engine/EngineTypes.h"
+#include "LocalizationTarget.generated.h"
+
+UENUM()
+enum class ELocalizationTargetStatus : uint8
+{
+	Unknown,
+	ConflictsPresent,
+	Clear
+};
+
+USTRUCT()
+struct FCultureStatistics
+{
+	GENERATED_USTRUCT_BODY()
+
+	FCultureStatistics()
+		: WordCount(0)
+	{
+	}
+
+	UPROPERTY(config, EditAnywhere, Category = "Culture")
+	FString CultureName;
+
+	UPROPERTY(Transient, EditAnywhere, Category = "Statistics")
+	uint32 WordCount;
+};
+
+USTRUCT()
+struct FLocalizationTargetSettings
+{
+	GENERATED_USTRUCT_BODY()
+
+	FLocalizationTargetSettings();
+
+	/* Unique name for the target within the context of the project. */
+	UPROPERTY(config, EditAnywhere, Category = "Target")
+	FString Name;
+
+	UPROPERTY(Transient, EditAnywhere, Category = "Target")
+	ELocalizationTargetStatus Status;
+
+	UPROPERTY(config, EditAnywhere, Category = "Target")
+	TArray<FString> TargetDependencies;
+
+	UPROPERTY(config, EditAnywhere, Category = "Target", AdvancedDisplay, meta=(FilePathFilter="manifest"))
+	TArray<FFilePath> AdditionalManifestDependencies;
+
+	UPROPERTY(config, EditAnywhere, Category = "Cultures")
+	FCultureStatistics NativeCultureStatistics;
+
+	UPROPERTY(config, EditAnywhere, Category = "Cultures")
+	TArray<FCultureStatistics> SupportedCulturesStatistics;
+
+	bool UpdateWordCountsFromCSV();
+	void UpdateStatusFromConflictReport();
+	bool Rename(const FString& NewName);
+	bool DeleteFiles(const FString* const Culture = nullptr) const;
+};
+
+UCLASS(Within=ProjectLocalizationSettings)
+class ULocalizationTarget : public UObject
+{
+	GENERATED_UCLASS_BODY()
+
+public:
+	UPROPERTY(EditAnywhere, Category = "Target")
+	FLocalizationTargetSettings Settings;
+
+public:
+#if WITH_EDITOR
+	virtual void PostEditChangeProperty(FPropertyChangedEvent& PropertyChangedEvent) override;
+#endif
+};
\ No newline at end of file
diff --git a/Engine/Source/Editor/LocalizationDashboard/Classes/ProjectLocalizationSettings.h b/Engine/Source/Editor/LocalizationDashboard/Classes/ProjectLocalizationSettings.h
new file mode 100644
--- /dev/null
+++ b/Engine/Source/Editor/LocalizationDashboard/Classes/ProjectLocalizationSettings.h
@@ -0,0 +1,30 @@
+// Copyright 1998-2014 Epic Games, Inc. All Rights Reserved.
+
+#pragma once
+
+#include "Engine/EngineTypes.h"
+#include "LocalizationTarget.h"
+#include "ProjectLocalizationSettings.generated.h"
+
+UCLASS(Config=Game, defaultconfig)
+class UProjectLocalizationSettings : public UObject
+{
+	GENERATED_UCLASS_BODY()
+
+public:
+	UPROPERTY(config, EditAnywhere, Category = "LocalizationServiceProvider")
+	FName LocalizationServiceProvider;
+
+	UPROPERTY(EditAnywhere, Category = "Targets")
+	TArray<ULocalizationTarget*> TargetObjects;
+
+private:
+	UPROPERTY(config)
+	TArray<FLocalizationTargetSettings> Targets;
+
+public:
+#if WITH_EDITOR
+	virtual void PostInitProperties() override;
+	virtual void PostEditChangeProperty(FPropertyChangedEvent& PropertyChangedEvent) override;
+#endif
+};
\ No newline at end of file
diff --git a/Engine/Source/Editor/LocalizationDashboard/LocalizationDashboard.Build.cs b/Engine/Source/Editor/LocalizationDashboard/LocalizationDashboard.Build.cs
new file mode 100644
--- /dev/null
+++ b/Engine/Source/Editor/LocalizationDashboard/LocalizationDashboard.Build.cs
@@ -0,0 +1,47 @@
+// Copyright 1998-2014 Epic Games, Inc. All Rights Reserved.
+
+using UnrealBuildTool;
+
+public class LocalizationDashboard : ModuleRules
+{
+	public LocalizationDashboard(TargetInfo Target)
+	{
+		PublicDependencyModuleNames.AddRange(
+			new string[]
+			{
+				"Core",
+                "CoreUObject",
+                "PropertyEditor",
+			}
+		);
+
+		PrivateDependencyModuleNames.AddRange(
+			new string[]
+			{
+				"Engine",
+				"Slate",
+				"SlateCore",
+				"InputCore",
+                "UnrealEd",
+                "EditorStyle",
+				"DesktopPlatform",
+                "TranslationEditor"
+			}
+		);
+
+		PrivateIncludePaths.AddRange(
+			new string[]
+			{
+				"Editor/LocalizationDashboard/Private",
+			}
+		);
+
+        PublicIncludePaths.AddRange(
+			new string[]
+			{
+				"Editor/LocalizationDashboard/Public",
+				"Editor/LocalizationDashboard/Classes",
+			}
+		);
+	}
+}
\ No newline at end of file
diff --git a/Engine/Source/Editor/LocalizationDashboard/Private/LocalizationCommandletExecution.cpp b/Engine/Source/Editor/LocalizationDashboard/Private/LocalizationCommandletExecution.cpp
new file mode 100644
--- /dev/null
+++ b/Engine/Source/Editor/LocalizationDashboard/Private/LocalizationCommandletExecution.cpp
@@ -0,0 +1,695 @@
+// Copyright 1998-2014 Epic Games, Inc. All Rights Reserved.
+
+#include "LocalizationDashboardPrivatePCH.h"
+#include "LocalizationCommandletExecution.h"
+#include "Classes/LocalizationTarget.h"
+#include "LocalizationConfigurationScript.h"
+#include "ISourceControlModule.h"
+#include "DesktopPlatformModule.h"
+#include "SThrobber.h"
+#include "Commandlets/CommandletHelpers.h"
+
+namespace
+{
+	class SLocalizationCommandletExecutor : public SWindow
+	{
+		SLATE_BEGIN_ARGS(SLocalizationCommandletExecutor)
+			: _Title( NSLOCTEXT("LocalizationDashboard", "LocalizationCommandletExecutorTitle", "Localization Commandlet") )
+		{}
+		SLATE_ATTRIBUTE(FText, Title)
+			SLATE_END_ARGS()
+
+	private:
+		struct FTaskListModel
+		{
+			enum class EState
+			{
+				Queued,
+				InProgress,
+				Failed,
+				Succeeded
+			};
+
+			FTaskListModel()
+				: State(EState::Queued)
+			{
+			}
+
+			LocalizationCommandletExecution::FTask Task;
+			EState State;
+			FString LogOutput;
+		};
+
+		friend class STaskRow;
+
+	public:
+		SLocalizationCommandletExecutor();
+		~SLocalizationCommandletExecutor();
+
+		void Construct(const FArguments& Arguments, const TArray<LocalizationCommandletExecution::FTask>& Tasks);
+		void Tick( const FGeometry& AllottedGeometry, const double InCurrentTime, const float InDeltaTime ) override;
+		bool WasSuccessful() const;
+		void Log(const FString& String);
+
+	private:
+		void ExecuteCommandlet(const TSharedRef<FTaskListModel>& TaskListModel);
+		void OnCommandletProcessCompletion(const int32 ReturnCode);
+		void CancelCommandlet();
+		void CleanUpProcessAndPump();
+
+		bool HasCompleted() const;
+
+		FText GetProgressMessageText() const;
+		TOptional<float> GetProgressPercentage() const;
+
+		TSharedRef<ITableRow> OnGenerateTaskListRow(TSharedPtr<FTaskListModel> TaskListModel, const TSharedRef<STableViewBase>& Table);
+		TSharedPtr<FTaskListModel> GetCurrentTaskToView() const;
+
+		FText GetLogString() const;
+
+		FReply OnCopyLogClicked();
+		void CopyLogToClipboard();
+
+		FReply OnSaveLogClicked();
+
+		FText GetCloseButtonText() const;
+		FReply OnCloseButtonClicked();
+
+	private:
+		int32 CurrentTaskIndex;
+		TArray< TSharedPtr<FTaskListModel> > TaskListModels;
+		TSharedPtr<SProgressBar> ProgressBar;
+		TSharedPtr< SListView< TSharedPtr<FTaskListModel> > > TaskListView;
+
+		struct
+		{
+			FCriticalSection CriticalSection;
+			FString String;
+		} PendingData;
+
+
+		TSharedPtr<FLocalizationCommandletProcess> CommandletProcess;
+		FRunnable* Runnable;
+		FRunnableThread* RunnableThread;
+	};
+
+	SLocalizationCommandletExecutor::SLocalizationCommandletExecutor()
+		: CurrentTaskIndex(INDEX_NONE)
+		, Runnable(nullptr)
+		, RunnableThread(nullptr)
+	{
+	}
+
+	SLocalizationCommandletExecutor::~SLocalizationCommandletExecutor()
+	{
+		CancelCommandlet();
+	}
+
+	void SLocalizationCommandletExecutor::Construct(const FArguments& Arguments, const TArray<LocalizationCommandletExecution::FTask>& Tasks)
+	{
+		for (const LocalizationCommandletExecution::FTask& Task : Tasks)
+		{
+			const TSharedRef<FTaskListModel> Model = MakeShareable(new FTaskListModel());
+			Model->Task = Task;
+			TaskListModels.Add(Model);
+		}
+
+		TSharedPtr<SScrollBar> VerticalScrollBar;
+		TSharedPtr<SScrollBar> HorizontalScrollBar;
+
+		SWindow::Construct(
+			SWindow::FArguments()
+			.SupportsMinimize(false)
+			.AutoCenter(EAutoCenter::PreferredWorkArea)
+			.ClientSize(FVector2D(600,400))
+			.ActivateWhenFirstShown(true)
+			.FocusWhenFirstShown(true)
+			.Content()
+			[
+				SNew(SVerticalBox)
+				+ SVerticalBox::Slot()
+				.AutoHeight()
+				.Padding(8.0, 16.0, 16.0, 0.0)
+				[
+					SNew(SVerticalBox)
+					+ SVerticalBox::Slot()
+					.AutoHeight()
+					[
+						SNew(STextBlock)
+						.Text(this, &SLocalizationCommandletExecutor::GetProgressMessageText)
+					]
+
+					+ SVerticalBox::Slot()
+						.AutoHeight()
+						.Padding(0.0, 4.0, 0.0, 0.0)
+						[
+							SAssignNew(ProgressBar, SProgressBar)
+							.Percent(this, &SLocalizationCommandletExecutor::GetProgressPercentage)
+						]
+				]
+				+ SVerticalBox::Slot()
+					.FillHeight(0.5)
+					.Padding(0.0, 32.0, 8.0, 0.0)
+					[
+						SNew(SBorder)
+						.BorderImage(FEditorStyle::GetBrush("ToolPanel.GroupBorder"))
+						.Padding(0.0f)
+						[
+							SAssignNew(TaskListView, SListView< TSharedPtr<FTaskListModel> >)
+							.HeaderRow
+							(
+							SNew(SHeaderRow)
+							+ SHeaderRow::Column("StatusIcon")
+							.DefaultLabel(FText::GetEmpty())
+							.FixedWidth(20.0)
+							+ SHeaderRow::Column("TaskName")
+							.DefaultLabel(NSLOCTEXT("LocalizationCommandletExecutor", "TaskListNameColumnLabel", "Task"))
+							.FillWidth(1.0)
+							)
+							.ListItemsSource(&TaskListModels)
+							.OnGenerateRow(this, &SLocalizationCommandletExecutor::OnGenerateTaskListRow)
+							.ItemHeight(24.0)
+							.SelectionMode(ESelectionMode::Single)
+						]
+					]
+				+ SVerticalBox::Slot()
+					.FillHeight(0.5)
+					.Padding(0.0, 32.0, 8.0, 0.0)
+					[
+						SNew(SBorder)
+						.BorderImage(FEditorStyle::GetBrush("ToolPanel.GroupBorder"))
+						.Padding(0.0f)
+						[
+							SNew(SHorizontalBox)
+							+SHorizontalBox::Slot()
+							.FillWidth(1.0f)
+							[
+								SNew(SVerticalBox)
+								+SVerticalBox::Slot()
+								.FillHeight(1.0f)
+								[
+									SNew(SMultiLineEditableText)
+									.TextStyle(FEditorStyle::Get(), "LocalizationDashboard.CommandletLog.Text")
+									.Text(this, &SLocalizationCommandletExecutor::GetLogString)
+									.IsReadOnly(true)
+									.HScrollBar(HorizontalScrollBar)
+									.VScrollBar(VerticalScrollBar)
+								]
+								+SVerticalBox::Slot()
+									.AutoHeight()
+									[
+										SAssignNew(HorizontalScrollBar, SScrollBar)
+										.Orientation(EOrientation::Orient_Horizontal)
+									]
+							]
+							+SHorizontalBox::Slot()
+								.AutoWidth()
+								[
+									SAssignNew(VerticalScrollBar, SScrollBar)
+									.Orientation(EOrientation::Orient_Vertical)
+								]
+						]
+					]
+				+ SVerticalBox::Slot()
+					.AutoHeight()
+					.Padding(0.0f, 5.0f, 0.0f, 0.0f)
+					[
+						SNew(SHorizontalBox)
+						+ SHorizontalBox::Slot()
+						.AutoWidth()
+						[
+							SNew(SButton)
+							.ContentPadding(FMargin(6.0f, 2.0f))
+							.Text(NSLOCTEXT("LocalizationCommandletExecutor", "CopyLogButtonText", "Copy Log"))
+							.ToolTipText(NSLOCTEXT("LocalizationCommandletExecutor", "CopyLogButtonTooltip", "Copy the logged text to the clipboard."))
+							.OnClicked(this, &SLocalizationCommandletExecutor::OnCopyLogClicked)
+						]
+						+ SHorizontalBox::Slot()
+							.AutoWidth()
+							[
+								SNew(SButton)
+								.ContentPadding(FMargin(6.0f, 2.0f))
+								.IsEnabled(false)
+								.Text(NSLOCTEXT("LocalizationCommandletExecutor", "SaveLogButtonText", "Save Log..."))
+								.ToolTipText(NSLOCTEXT("LocalizationCommandletExecutor", "SaveLogButtonToolTip", "Save the logged text to a file."))
+								.OnClicked(this, &SLocalizationCommandletExecutor::OnSaveLogClicked)
+							]
+						+ SHorizontalBox::Slot()
+							.AutoWidth()
+							[
+								SNew(SButton)
+								.ContentPadding(FMargin(6.0f, 2.0f))
+								.OnClicked(this, &SLocalizationCommandletExecutor::OnCloseButtonClicked)
+								[
+									SNew(STextBlock)
+									.Text(this, &SLocalizationCommandletExecutor::GetCloseButtonText)
+								]
+							]
+					]
+			]);
+
+		if(TaskListModels.Num() > 0)
+		{
+			CurrentTaskIndex = 0;
+			ExecuteCommandlet(TaskListModels[CurrentTaskIndex].ToSharedRef());
+		}
+	}
+
+	void SLocalizationCommandletExecutor::Tick( const FGeometry& AllottedGeometry, const double InCurrentTime, const float InDeltaTime )
+	{
+		SCompoundWidget::Tick(AllottedGeometry, InCurrentTime, InDeltaTime);
+
+		// Poll for log output data.
+		if (!PendingData.String.IsEmpty())
+		{
+			FString String;
+
+			// Copy the pending data string to the local string
+			{
+				FScopeLock ScopeLock(&PendingData.CriticalSection);
+				String = PendingData.String;
+				PendingData.String.Empty();
+			}
+
+			// Forward string to proper log.
+			const TSharedPtr<FTaskListModel> CurrentTaskModel = TaskListModels[CurrentTaskIndex];
+			CurrentTaskModel->LogOutput.Append(String);
+		}
+
+		// On Task Completed.
+		if (CommandletProcess.IsValid())
+		{
+			FProcHandle CurrentProcessHandle = CommandletProcess->GetHandle();
+			int32 ReturnCode;
+			if (CurrentProcessHandle.IsValid() && FPlatformProcess::GetProcReturnCode(CurrentProcessHandle, &ReturnCode))
+			{
+				OnCommandletProcessCompletion(ReturnCode);
+			}
+		}
+	}
+
+	bool SLocalizationCommandletExecutor::WasSuccessful() const
+	{
+		return HasCompleted() && !TaskListModels.ContainsByPredicate([](const TSharedPtr<FTaskListModel>& TaskListModel){return TaskListModel->State != FTaskListModel::EState::Succeeded;});
+	}
+
+	void SLocalizationCommandletExecutor::Log(const FString& String)
+	{
+		FScopeLock ScopeLock(&PendingData.CriticalSection);
+		PendingData.String += String;
+	}
+
+	void SLocalizationCommandletExecutor::OnCommandletProcessCompletion(const int32 ReturnCode)
+	{
+		CleanUpProcessAndPump();
+
+		// Handle return code.
+		TSharedPtr<FTaskListModel> CurrentTaskModel = TaskListModels[CurrentTaskIndex];
+		// Zero code is successful.
+		if (ReturnCode == 0)
+		{
+			CurrentTaskModel->State = FTaskListModel::EState::Succeeded;
+
+			++CurrentTaskIndex;
+
+			// Begin new task if possible.
+			if (TaskListModels.IsValidIndex(CurrentTaskIndex))
+			{
+				CurrentTaskModel = TaskListModels[CurrentTaskIndex];
+				ExecuteCommandlet(CurrentTaskModel.ToSharedRef());
+			}
+		}
+		// Non-zero is a failure.
+		else
+		{
+			CurrentTaskModel->State = FTaskListModel::EState::Failed;
+		}
+	}
+
+	void SLocalizationCommandletExecutor::ExecuteCommandlet(const TSharedRef<FTaskListModel>& TaskListModel)
+	{
+		CommandletProcess = FLocalizationCommandletProcess::Execute(TaskListModel->Task.ScriptPath);
+		if (CommandletProcess.IsValid())
+		{
+			TaskListModel->State = FTaskListModel::EState::InProgress;
+		}
+		else
+		{
+			TaskListModel->State = FTaskListModel::EState::Failed;
+			return;
+		}
+
+		class FCommandletLogPump : public FRunnable
+		{
+		public:
+			FCommandletLogPump(void* const InReadPipe, const FProcHandle& InCommandletProcessHandle, const TSharedRef<SLocalizationCommandletExecutor>& InCommandletWidget)
+				: ReadPipe(InReadPipe)
+				, CommandletProcessHandle(InCommandletProcessHandle)
+				, CommandletWidget(InCommandletWidget)
+			{
+			}
+
+			uint32 Run() override
+			{
+				for(;;)
+				{
+					// Read from pipe.
+					const FString PipeString = FPlatformProcess::ReadPipe(ReadPipe);
+
+					// Process buffer.
+					if (!PipeString.IsEmpty())
+					{
+						// Add strings to log.
+						const TSharedPtr<SLocalizationCommandletExecutor> Widget = CommandletWidget.Pin();
+						if (Widget.IsValid())
+						{
+							Widget->Log(PipeString);
+						}
+					}
+
+					// If the process isn't running and there's no data in the pipe, we're done.
+					if (!FPlatformProcess::IsProcRunning(CommandletProcessHandle) && PipeString.IsEmpty())
+					{
+						break;
+					}
+
+					// Sleep.
+					FPlatformProcess::Sleep(0.0f);
+				}
+
+				int32 ReturnCode = 0;
+				return FPlatformProcess::GetProcReturnCode(CommandletProcessHandle, &ReturnCode) ? ReturnCode : -1;
+			}
+
+		private:
+			void* const ReadPipe;
+			FProcHandle CommandletProcessHandle;
+			const TWeakPtr<SLocalizationCommandletExecutor> CommandletWidget;
+		};
+
+		// Launch runnable thread.
+		Runnable = new FCommandletLogPump(CommandletProcess->GetReadPipe(), CommandletProcess->GetHandle(), SharedThis(this));
+		RunnableThread = FRunnableThread::Create(Runnable, TEXT("Localization Commandlet Log Pump Thread"));
+	}
+
+	void SLocalizationCommandletExecutor::CancelCommandlet()
+	{
+		CleanUpProcessAndPump();
+	}
+
+	void SLocalizationCommandletExecutor::CleanUpProcessAndPump()
+	{
+		if (CommandletProcess.IsValid())
+		{
+			FProcHandle CommandletProcessHandle = CommandletProcess->GetHandle();
+			if (CommandletProcessHandle.IsValid() && FPlatformProcess::IsProcRunning(CommandletProcessHandle))
+			{
+				FPlatformProcess::TerminateProc(CommandletProcessHandle, true);
+			}
+			CommandletProcess.Reset();
+		}
+
+		if (RunnableThread)
+		{
+			RunnableThread->WaitForCompletion();
+			delete RunnableThread;
+			RunnableThread = nullptr;
+		}
+
+		if (Runnable)
+		{
+			delete Runnable;
+			Runnable = nullptr;
+		}
+
+	}
+
+	bool SLocalizationCommandletExecutor::HasCompleted() const
+	{
+		return CurrentTaskIndex == TaskListModels.Num();
+	}
+
+	FText SLocalizationCommandletExecutor::GetProgressMessageText() const
+	{
+		return TaskListModels.IsValidIndex(CurrentTaskIndex) ? TaskListModels[CurrentTaskIndex]->Task.Name : FText::GetEmpty();
+	}
+
+	TOptional<float> SLocalizationCommandletExecutor::GetProgressPercentage() const
+	{
+		return TOptional<float>(float(CurrentTaskIndex) / float(TaskListModels.Num()));
+	}
+
+	class STaskRow : public SMultiColumnTableRow< TSharedPtr<SLocalizationCommandletExecutor::FTaskListModel> >
+	{
+	public:
+		void Construct(const FTableRowArgs& InArgs, const TSharedRef<STableViewBase>& OwnerTableView, const TSharedRef<SLocalizationCommandletExecutor::FTaskListModel>& InTaskListModel);
+		TSharedRef<SWidget> GenerateWidgetForColumn(const FName& ColumnName);
+
+	private:
+		FSlateColor HandleIconColorAndOpacity() const;
+		const FSlateBrush* HandleIconImage() const;
+		EVisibility HandleThrobberVisibility() const;
+
+	private:
+		TSharedPtr<SLocalizationCommandletExecutor::FTaskListModel> TaskListModel;
+	};
+
+	void STaskRow::Construct(const FTableRowArgs& InArgs, const TSharedRef<STableViewBase>& OwnerTableView, const TSharedRef<SLocalizationCommandletExecutor::FTaskListModel>& InTaskListModel)
+	{
+		TaskListModel = InTaskListModel;
+
+		FSuperRowType::Construct(InArgs, OwnerTableView);
+	}
+
+	TSharedRef<SWidget> STaskRow::GenerateWidgetForColumn(const FName& ColumnName)
+	{
+		if (ColumnName == "StatusIcon")
+		{
+			return SNew(SOverlay)
+
+				+ SOverlay::Slot()
+				.HAlign(HAlign_Center)
+				.VAlign(VAlign_Center)
+				[
+					SNew(SThrobber)
+					.Animate(SThrobber::VerticalAndOpacity)
+					.NumPieces(1)
+					.Visibility(this, &STaskRow::HandleThrobberVisibility)
+				]
+			+ SOverlay::Slot()
+				.HAlign(HAlign_Center)
+				.VAlign(VAlign_Center)
+				[
+					SNew(SImage)
+					.ColorAndOpacity(this, &STaskRow::HandleIconColorAndOpacity)
+					.Image(this, &STaskRow::HandleIconImage)
+				];
+		}
+		else if (ColumnName == "TaskName")
+		{
+			return SNew(STextBlock)
+				.Text(TaskListModel->Task.Name);
+		}
+		else
+		{
+			return SNullWidget::NullWidget;
+		}
+	}
+
+	FSlateColor STaskRow::HandleIconColorAndOpacity( ) const
+	{
+		if (TaskListModel.IsValid())
+		{
+			switch(TaskListModel->State)
+			{
+			case SLocalizationCommandletExecutor::FTaskListModel::EState::InProgress:
+				return FLinearColor::Yellow;
+				break;
+			case SLocalizationCommandletExecutor::FTaskListModel::EState::Succeeded:
+				return FLinearColor::Green;
+				break;
+			case SLocalizationCommandletExecutor::FTaskListModel::EState::Failed:
+				return FLinearColor::Red;
+				break;
+			}
+		}
+
+		return FSlateColor::UseForeground();
+	}
+
+	const FSlateBrush* STaskRow::HandleIconImage( ) const
+	{
+		if (TaskListModel.IsValid())
+		{
+			switch(TaskListModel->State)
+			{
+			case SLocalizationCommandletExecutor::FTaskListModel::EState::Succeeded:
+				return FEditorStyle::GetBrush("Symbols.Check");
+				break;
+			case SLocalizationCommandletExecutor::FTaskListModel::EState::Failed:
+				return FEditorStyle::GetBrush("Icons.Cross");
+				break;
+			}
+		}
+
+		return NULL;
+	}
+
+	EVisibility STaskRow::HandleThrobberVisibility( ) const
+	{
+		if (TaskListModel.IsValid())
+		{
+			switch(TaskListModel->State)
+			{
+			case SLocalizationCommandletExecutor::FTaskListModel::EState::InProgress:
+				return EVisibility::Visible;
+				break;
+			}
+		}
+
+		return EVisibility::Hidden;
+	}
+
+	TSharedRef<ITableRow> SLocalizationCommandletExecutor::OnGenerateTaskListRow(TSharedPtr<FTaskListModel> TaskListModel, const TSharedRef<STableViewBase>& Table)
+	{
+		return SNew(STaskRow, Table, TaskListModel.ToSharedRef());
+	}
+
+	TSharedPtr<SLocalizationCommandletExecutor::FTaskListModel> SLocalizationCommandletExecutor::GetCurrentTaskToView() const
+	{
+		if (TaskListView.IsValid())
+		{
+			const TArray< TSharedPtr<FTaskListModel> > Selection = TaskListView->GetSelectedItems();
+			return Selection.Num() > 0 ? Selection.Top() : nullptr;
+		}
+		return nullptr;
+	}
+
+	FText SLocalizationCommandletExecutor::GetLogString() const
+	{
+		const TSharedPtr<SLocalizationCommandletExecutor::FTaskListModel> TaskToView = GetCurrentTaskToView();
+		return TaskToView.IsValid() ? FText::FromString(TaskToView->LogOutput) : FText::GetEmpty();
+	}
+
+	FReply SLocalizationCommandletExecutor::OnCopyLogClicked()
+	{
+		CopyLogToClipboard();
+
+		return FReply::Handled();
+	}
+
+	void SLocalizationCommandletExecutor::CopyLogToClipboard()
+	{
+		FPlatformMisc::ClipboardCopy(*(GetLogString().ToString()));
+	}
+
+	FReply SLocalizationCommandletExecutor::OnSaveLogClicked()
+	{
+		const FString TextFileDescription = NSLOCTEXT("LocalizationCommandletExecutor", "TextFileDescription", "Text File").ToString();
+		const FString TextFileExtension = TEXT("txt");
+		const FString TextFileExtensionWildcard = FString::Printf(TEXT("*.%s"), *TextFileExtension);
+		const FString FileTypes = FString::Printf(TEXT("%s (%s)|%s"), *TextFileDescription, *TextFileExtensionWildcard, *TextFileExtensionWildcard);
+		const FString DefaultFilename = FString::Printf(TEXT("%s.%s"), TEXT("Log"), *TextFileExtension);
+		const FString DefaultPath = FPaths::GameSavedDir();
+
+		TArray<FString> SaveFilenames;
+		IDesktopPlatform* DesktopPlatform = FDesktopPlatformModule::Get();
+
+		// Prompt the user for the filename
+		if (DesktopPlatform)
+		{
+			void* ParentWindowWindowHandle = NULL;
+
+			const TSharedPtr<SWindow>& ParentWindow = FSlateApplication::Get().FindWidgetWindow(SharedThis(this));
+			if (ParentWindow.IsValid() && ParentWindow->GetNativeWindow().IsValid())
+			{
+				ParentWindowWindowHandle = ParentWindow->GetNativeWindow()->GetOSWindowHandle();
+			}
+
+			if (DesktopPlatform->SaveFileDialog(
+				ParentWindowWindowHandle,
+				NSLOCTEXT("LocalizationDashboard", "SaveLogDialogTitle", "Save Log to File").ToString(),
+				DefaultPath,
+				DefaultFilename,
+				FileTypes,
+				EFileDialogFlags::None,
+				SaveFilenames
+				))
+			{
+				// Save to file.
+				FFileHelper::SaveStringToFile( GetLogString().ToString(), *(SaveFilenames.Last()) );
+			}
+		}
+
+		return FReply::Handled();
+	}
+
+	FText SLocalizationCommandletExecutor::GetCloseButtonText() const
+	{
+		return HasCompleted() ? NSLOCTEXT("LocalizationDashboard", "OkayButtonText", "Okay") : NSLOCTEXT("LocalizationDashboard", "CancelButtonText", "Cancel");
+	}
+
+	FReply SLocalizationCommandletExecutor::OnCloseButtonClicked()
+	{
+		if (!HasCompleted())
+		{
+			CancelCommandlet();
+		}
+		RequestDestroyWindow();
+		return FReply::Handled();
+	}
+}
+
+bool LocalizationCommandletExecution::Execute(const TSharedRef<SWindow>& ParentWindow, const FText& Title, const TArray<FTask>& Tasks)
+{
+	const TSharedRef<SLocalizationCommandletExecutor> CommandletWindow = SNew(SLocalizationCommandletExecutor, Tasks)
+		.Title(Title);
+	FSlateApplication::Get().AddModalWindow(CommandletWindow, ParentWindow, false);
+	return CommandletWindow->WasSuccessful();
+}
+
+TSharedPtr<FLocalizationCommandletProcess> FLocalizationCommandletProcess::Execute(const FString& ConfigFilePath)
+{
+	// Create pipes.
+	void* ReadPipe;
+	void* WritePipe;
+	if (!FPlatformProcess::CreatePipe(ReadPipe, WritePipe))
+	{
+		return nullptr;
+	}
+
+	// Create process.
+	FString CommandletArguments;
+
+	const FString ConfigFileRelativeToGameDir = LocalizationConfigurationScript::MakePathRelativeToProjectDirectory(ConfigFilePath);
+	CommandletArguments = FString::Printf( TEXT("-config=\"%s\""), *ConfigFileRelativeToGameDir );
+
+	ISourceControlModule& SourceControlModule = ISourceControlModule::Get();
+	if (SourceControlModule.IsEnabled())
+	{
+		if (!SourceControlModule.GetProvider().IsAvailable())
+		{
+			CommandletArguments = FString::Printf( TEXT("%s -EnableSCC -DisableSCCSubmit"), *CommandletArguments );
+		}
+	}
+
+	FProcHandle CommandletProcessHandle = CommandletHelpers::CreateCommandletProcess(TEXT("GatherText"), *FPaths::GetProjectFilePath(), *CommandletArguments, true, true, true, nullptr, 0, nullptr, WritePipe);
+
+	// Close pipes if process failed.
+	if (!CommandletProcessHandle.IsValid())
+	{
+		FPlatformProcess::ClosePipe(ReadPipe, WritePipe);
+		return nullptr;
+	}
+
+	return MakeShareable(new FLocalizationCommandletProcess(ReadPipe, WritePipe, CommandletProcessHandle));
+}
+
+FLocalizationCommandletProcess::~FLocalizationCommandletProcess()
+{
+	if (FPlatformProcess::IsProcRunning(ProcessHandle))
+	{
+		FPlatformProcess::TerminateProc(ProcessHandle);
+	}
+	FPlatformProcess::ClosePipe(ReadPipe, WritePipe);
+}
\ No newline at end of file
diff --git a/Engine/Source/Editor/LocalizationDashboard/Private/LocalizationCommandletExecution.h b/Engine/Source/Editor/LocalizationDashboard/Private/LocalizationCommandletExecution.h
new file mode 100644
--- /dev/null
+++ b/Engine/Source/Editor/LocalizationDashboard/Private/LocalizationCommandletExecution.h
@@ -0,0 +1,54 @@
+// Copyright 1998-2014 Epic Games, Inc. All Rights Reserved.
+#pragma once 
+
+#include "LocalizationConfigurationScript.h"
+
+namespace LocalizationCommandletExecution
+{
+	struct FTask
+	{
+		FTask() {}
+
+		FTask(const FText& InName, const FString& InScriptPath)
+			: Name(InName)
+			, ScriptPath(InScriptPath)
+		{}
+
+		FText Name;
+		FString ScriptPath;
+	};
+
+	bool Execute(const TSharedRef<SWindow>& ParentWindow, const FText& Title, const TArray<FTask>& Tasks);
+};
+
+class FLocalizationCommandletProcess : public TSharedFromThis<FLocalizationCommandletProcess>
+{
+public:
+	static TSharedPtr<FLocalizationCommandletProcess> Execute(const FString& ConfigFilePath);
+
+private:
+	FLocalizationCommandletProcess(void* const InReadPipe, void* const InWritePipe, const FProcHandle InProcessHandle)
+		: ReadPipe(InReadPipe)
+		, WritePipe(InWritePipe)
+		, ProcessHandle(InProcessHandle)
+	{
+	}
+
+public:
+	~FLocalizationCommandletProcess();
+
+	void* GetReadPipe() const
+	{
+		return ReadPipe;
+	}
+
+	FProcHandle& GetHandle()
+	{
+		return ProcessHandle;
+	}
+
+private:
+	void* const ReadPipe;
+	void* const WritePipe;
+	FProcHandle ProcessHandle;
+};
\ No newline at end of file
diff --git a/Engine/Source/Editor/LocalizationDashboard/Private/LocalizationCommandletTasks.cpp b/Engine/Source/Editor/LocalizationDashboard/Private/LocalizationCommandletTasks.cpp
new file mode 100644
--- /dev/null
+++ b/Engine/Source/Editor/LocalizationDashboard/Private/LocalizationCommandletTasks.cpp
@@ -0,0 +1,205 @@
+// Copyright 1998-2014 Epic Games, Inc. All Rights Reserved.
+
+#include "LocalizationDashboardPrivatePCH.h"
+#include "LocalizationCommandletTasks.h"
+#include "LocalizationCommandletExecution.h"
+#include "LocalizationConfigurationScript.h"
+
+bool LocalizationCommandletTasks::GatherTargets(const TSharedRef<SWindow>& ParentWindow, const TArray<FLocalizationTargetSettings*>& TargetsSettings)
+{
+	TArray<LocalizationCommandletExecution::FTask> Tasks;
+
+	for (FLocalizationTargetSettings* TargetSettings : TargetsSettings)
+	{
+		FFormatNamedArguments Arguments;
+		Arguments.Add(TEXT("TargetName"), FText::FromString(TargetSettings->Name));
+
+		const FText GatherTaskName = FText::Format(NSLOCTEXT("LocalizationDashboard", "GatherTaskNameFormat", "Gather Text for {TargetName}"), Arguments);
+		const FString GatherScriptPath = LocalizationConfigurationScript::GetGatherScriptPath(*TargetSettings);
+		LocalizationConfigurationScript::GenerateGatherScript(*TargetSettings).Write(GatherScriptPath);
+		Tasks.Add(LocalizationCommandletExecution::FTask(GatherTaskName, GatherScriptPath));
+
+		const FText ReportTaskName = FText::Format(NSLOCTEXT("LocalizationDashboard", "ReportTaskNameFormat", "Generate Reports for {TargetName}"), Arguments);
+		const FString ReportScriptPath = LocalizationConfigurationScript::GetReportScriptPath(*TargetSettings);
+		LocalizationConfigurationScript::GenerateReportScript(*TargetSettings).Write(ReportScriptPath);
+		Tasks.Add(LocalizationCommandletExecution::FTask(ReportTaskName, ReportScriptPath));
+	}
+
+	return LocalizationCommandletExecution::Execute(ParentWindow, NSLOCTEXT("LocalizationDashboard", "GatherAllTargetsWindowTitle", "Gather Text for All Targets"), Tasks);
+}
+
+bool LocalizationCommandletTasks::GatherTarget(const TSharedRef<SWindow>& ParentWindow, FLocalizationTargetSettings& TargetSettings)
+{
+	TArray<LocalizationCommandletExecution::FTask> Tasks;
+
+	const FString GatherScriptPath = LocalizationConfigurationScript::GetGatherScriptPath(TargetSettings);
+	LocalizationConfigurationScript::GenerateGatherScript(TargetSettings).Write(GatherScriptPath);
+	Tasks.Add(LocalizationCommandletExecution::FTask(NSLOCTEXT("LocalizationDashboard", "GatherTaskName", "Gather Text"), GatherScriptPath));
+
+	const FString ReportScriptPath = LocalizationConfigurationScript::GetReportScriptPath(TargetSettings);
+	LocalizationConfigurationScript::GenerateReportScript(TargetSettings).Write(ReportScriptPath);
+	Tasks.Add(LocalizationCommandletExecution::FTask(NSLOCTEXT("LocalizationDashboard", "ReportTaskName", "Generate Reports"), ReportScriptPath));
+
+	FFormatNamedArguments Arguments;
+	Arguments.Add(TEXT("TargetName"), FText::FromString(TargetSettings.Name));
+
+	const FText WindowTitle = FText::Format(NSLOCTEXT("LocalizationDashboard", "GatherTargetWindowTitle", "Gather Text for Target {TargetName}"), Arguments);
+	return LocalizationCommandletExecution::Execute(ParentWindow, WindowTitle, Tasks);
+}
+
+
+bool LocalizationCommandletTasks::ImportTargets(const TSharedRef<SWindow>& ParentWindow, const TArray<FLocalizationTargetSettings*>& TargetsSettings)
+{
+	TArray<LocalizationCommandletExecution::FTask> Tasks;
+
+	for (FLocalizationTargetSettings* TargetSettings : TargetsSettings)
+	{
+		FFormatNamedArguments Arguments;
+		Arguments.Add(TEXT("TargetName"), FText::FromString(TargetSettings->Name));
+
+		const FText ImportTaskName = FText::Format(NSLOCTEXT("LocalizationDashboard", "ImportTaskNameFormat", "Import Translations for {TargetName}"), Arguments);
+		const FString ImportScriptPath = LocalizationConfigurationScript::GetImportScriptPath(*TargetSettings);
+		LocalizationConfigurationScript::GenerateImportScript(*TargetSettings).Write(ImportScriptPath);
+		Tasks.Add(LocalizationCommandletExecution::FTask(ImportTaskName, ImportScriptPath));
+
+		const FText ReportTaskName = FText::Format(NSLOCTEXT("LocalizationDashboard", "ReportTaskNameFormat", "Generate Reports for {TargetName}"), Arguments);
+		const FString ReportScriptPath = LocalizationConfigurationScript::GetReportScriptPath(*TargetSettings);
+		LocalizationConfigurationScript::GenerateReportScript(*TargetSettings).Write(ReportScriptPath);
+		Tasks.Add(LocalizationCommandletExecution::FTask(ReportTaskName, ReportScriptPath));
+	}
+
+	return LocalizationCommandletExecution::Execute(ParentWindow, NSLOCTEXT("LocalizationDashboard", "ImportForAllTargetsWindowTitle", "Import Translations for All Targets"), Tasks);
+}
+
+bool LocalizationCommandletTasks::ImportTarget(const TSharedRef<SWindow>& ParentWindow, FLocalizationTargetSettings& TargetSettings)
+{
+	TArray<LocalizationCommandletExecution::FTask> Tasks;
+
+	const FString ImportScriptPath = LocalizationConfigurationScript::GetImportScriptPath(TargetSettings);
+	LocalizationConfigurationScript::GenerateImportScript(TargetSettings).Write(ImportScriptPath);
+	Tasks.Add(LocalizationCommandletExecution::FTask(NSLOCTEXT("LocalizationDashboard", "ImportTaskName", "Import Translations"), ImportScriptPath));
+
+	const FString ReportScriptPath = LocalizationConfigurationScript::GetReportScriptPath(TargetSettings);
+	LocalizationConfigurationScript::GenerateReportScript(TargetSettings).Write(ReportScriptPath);
+	Tasks.Add(LocalizationCommandletExecution::FTask(NSLOCTEXT("LocalizationDashboard", "ReportTaskName", "Generate Reports"), ReportScriptPath));
+
+	FFormatNamedArguments Arguments;
+	Arguments.Add(TEXT("TargetName"), FText::FromString(TargetSettings.Name));
+
+	const FText WindowTitle = FText::Format(NSLOCTEXT("LocalizationDashboard", "ImportForTargetWindowTitle", "Import Translations for Target {TargetName}"), Arguments);
+	return LocalizationCommandletExecution::Execute(ParentWindow, WindowTitle, Tasks);
+}
+
+bool LocalizationCommandletTasks::ImportCulture(const TSharedRef<SWindow>& ParentWindow, FLocalizationTargetSettings& TargetSettings, const FString& CultureName)
+{
+	FCulturePtr Culture = FInternationalization::Get().GetCulture(CultureName);
+	if (!Culture.IsValid())
+	{
+		return false;
+	}
+
+	TArray<LocalizationCommandletExecution::FTask> Tasks;
+
+	const FString ImportScriptPath = LocalizationConfigurationScript::GetImportScriptPath(TargetSettings, &CultureName);
+	LocalizationConfigurationScript::GenerateImportScript(TargetSettings, &CultureName).Write(ImportScriptPath);
+	Tasks.Add(LocalizationCommandletExecution::FTask(NSLOCTEXT("LocalizationDashboard", "ImportTaskName", "Import Translations"), ImportScriptPath));
+
+	const FString ReportScriptPath = LocalizationConfigurationScript::GetReportScriptPath(TargetSettings);
+	LocalizationConfigurationScript::GenerateReportScript(TargetSettings).Write(ReportScriptPath);
+	Tasks.Add(LocalizationCommandletExecution::FTask(NSLOCTEXT("LocalizationDashboard", "ReportTaskName", "Generate Reports"), ReportScriptPath));
+
+	FFormatNamedArguments Arguments;
+	Arguments.Add(TEXT("CultureName"), FText::FromString(Culture->GetDisplayName()));
+	Arguments.Add(TEXT("TargetName"), FText::FromString(TargetSettings.Name));
+
+	const FText WindowTitle = FText::Format(NSLOCTEXT("LocalizationDashboard", "ImportCultureForTargetWindowTitle", "Import {CultureName} Translations for Target {TargetName}"), Arguments);
+	return LocalizationCommandletExecution::Execute(ParentWindow, WindowTitle, Tasks);
+}
+
+bool LocalizationCommandletTasks::ExportTargets(const TSharedRef<SWindow>& ParentWindow, const TArray<FLocalizationTargetSettings*>& TargetsSettings)
+{
+	TArray<LocalizationCommandletExecution::FTask> Tasks;
+
+	for (FLocalizationTargetSettings* TargetSettings : TargetsSettings)
+	{
+		FFormatNamedArguments Arguments;
+		Arguments.Add(TEXT("TargetName"), FText::FromString(TargetSettings->Name));
+
+		const FText ExportTaskName = FText::Format(NSLOCTEXT("LocalizationDashboard", "ExportTaskNameFormat", "Export Translations for {TargetName}"), Arguments);
+		const FString ExportScriptPath = LocalizationConfigurationScript::GetExportScriptPath(*TargetSettings);
+		LocalizationConfigurationScript::GenerateExportScript(*TargetSettings).Write(ExportScriptPath);
+		Tasks.Add(LocalizationCommandletExecution::FTask(ExportTaskName, ExportScriptPath));
+	}
+
+	return LocalizationCommandletExecution::Execute(ParentWindow, NSLOCTEXT("LocalizationDashboard", "ExportForAllTargetsWindowTitle", "Export Translations for All Targets"), Tasks);
+}
+
+bool LocalizationCommandletTasks::ExportTarget(const TSharedRef<SWindow>& ParentWindow, FLocalizationTargetSettings& TargetSettings)
+{
+	TArray<LocalizationCommandletExecution::FTask> Tasks;
+
+	const FString ExportScriptPath = LocalizationConfigurationScript::GetExportScriptPath(TargetSettings);
+	LocalizationConfigurationScript::GenerateExportScript(TargetSettings).Write(ExportScriptPath);
+	Tasks.Add(LocalizationCommandletExecution::FTask(NSLOCTEXT("LocalizationDashboard", "ExportTaskName", "Export Translations"), ExportScriptPath));
+
+	FFormatNamedArguments Arguments;
+	Arguments.Add(TEXT("TargetName"), FText::FromString(TargetSettings.Name));
+
+	const FText WindowTitle = FText::Format(NSLOCTEXT("LocalizationDashboard", "ExportForTargetWindowTitle", "Export Translations for Target {TargetName}"), Arguments);
+	return LocalizationCommandletExecution::Execute(ParentWindow, WindowTitle, Tasks);
+}
+
+bool LocalizationCommandletTasks::ExportCulture(const TSharedRef<SWindow>& ParentWindow, FLocalizationTargetSettings& TargetSettings, const FString& CultureName)
+{
+	FCulturePtr Culture = FInternationalization::Get().GetCulture(CultureName);
+	if (!Culture.IsValid())
+	{
+		return false;
+	}
+
+	TArray<LocalizationCommandletExecution::FTask> Tasks;
+
+	const FString ExportScriptPath = LocalizationConfigurationScript::GetExportScriptPath(TargetSettings, &CultureName);
+	LocalizationConfigurationScript::GenerateExportScript(TargetSettings, &CultureName).Write(ExportScriptPath);
+	Tasks.Add(LocalizationCommandletExecution::FTask(NSLOCTEXT("LocalizationDashboard", "ExportTaskName", "Export Translations"), ExportScriptPath));
+
+	FFormatNamedArguments Arguments;
+	Arguments.Add(TEXT("CultureName"), FText::FromString(Culture->GetDisplayName()));
+	Arguments.Add(TEXT("TargetName"), FText::FromString(TargetSettings.Name));
+
+	const FText WindowTitle = FText::Format(NSLOCTEXT("LocalizationDashboard", "ExportCultureForTargetWindowTitle", "Export {CultureName} Translations for Target {TargetName}"), Arguments);
+	return LocalizationCommandletExecution::Execute(ParentWindow, WindowTitle, Tasks);
+}
+
+bool LocalizationCommandletTasks::GenerateReportsForTargets(const TSharedRef<SWindow>& ParentWindow, const TArray<FLocalizationTargetSettings*>& TargetsSettings)
+{
+	TArray<LocalizationCommandletExecution::FTask> Tasks;
+
+	for (FLocalizationTargetSettings* TargetSettings : TargetsSettings)
+	{
+		FFormatNamedArguments Arguments;
+		Arguments.Add(TEXT("TargetName"), FText::FromString(TargetSettings->Name));
+
+		const FText ReportTaskName = FText::Format(NSLOCTEXT("LocalizationDashboard", "ReportTaskNameFormat", "Generate Reports for {TargetName}"), Arguments);
+		const FString ReportScriptPath = LocalizationConfigurationScript::GetReportScriptPath(*TargetSettings);
+		LocalizationConfigurationScript::GenerateReportScript(*TargetSettings).Write(ReportScriptPath);
+		Tasks.Add(LocalizationCommandletExecution::FTask(ReportTaskName, ReportScriptPath));
+	}
+
+	return LocalizationCommandletExecution::Execute(ParentWindow, NSLOCTEXT("LocalizationDashboard", "GenerateReportsForAllTargetsWindowTitle", "Generate Reports for All Targets"), Tasks);
+}
+
+bool LocalizationCommandletTasks::GenerateReportsForTarget(const TSharedRef<SWindow>& ParentWindow, FLocalizationTargetSettings& TargetSettings)
+{
+	TArray<LocalizationCommandletExecution::FTask> Tasks;
+
+	const FString ReportScriptPath = LocalizationConfigurationScript::GetReportScriptPath(TargetSettings);
+	LocalizationConfigurationScript::GenerateReportScript(TargetSettings).Write(ReportScriptPath);
+	Tasks.Add(LocalizationCommandletExecution::FTask(NSLOCTEXT("LocalizationDashboard", "ReportTaskName", "Generate Reports"), ReportScriptPath));
+
+	FFormatNamedArguments Arguments;
+	Arguments.Add(TEXT("TargetName"), FText::FromString(TargetSettings.Name));
+
+	const FText WindowTitle = FText::Format(NSLOCTEXT("LocalizationDashboard", "GenerateReportsForTargetWindowTitle", "Generate Reports for Target {TargetName}"), Arguments);
+	return LocalizationCommandletExecution::Execute(ParentWindow, WindowTitle, Tasks);
+}
\ No newline at end of file
diff --git a/Engine/Source/Editor/LocalizationDashboard/Private/LocalizationCommandletTasks.h b/Engine/Source/Editor/LocalizationDashboard/Private/LocalizationCommandletTasks.h
new file mode 100644
--- /dev/null
+++ b/Engine/Source/Editor/LocalizationDashboard/Private/LocalizationCommandletTasks.h
@@ -0,0 +1,21 @@
+// Copyright 1998-2014 Epic Games, Inc. All Rights Reserved.
+#pragma once
+
+#include "Classes/LocalizationTarget.h"
+
+namespace LocalizationCommandletTasks
+{
+	bool GatherTargets(const TSharedRef<SWindow>& ParentWindow, const TArray<FLocalizationTargetSettings*>& TargetsSettings);
+	bool GatherTarget(const TSharedRef<SWindow>& ParentWindow, FLocalizationTargetSettings& TargetSettings);
+
+	bool ImportTargets(const TSharedRef<SWindow>& ParentWindow, const TArray<FLocalizationTargetSettings*>& TargetsSettings);
+	bool ImportTarget(const TSharedRef<SWindow>& ParentWindow, FLocalizationTargetSettings& TargetSettings);
+	bool ImportCulture(const TSharedRef<SWindow>& ParentWindow, FLocalizationTargetSettings& TargetSettings, const FString& CultureName);
+
+	bool ExportTargets(const TSharedRef<SWindow>& ParentWindow, const TArray<FLocalizationTargetSettings*>& TargetsSettings);
+	bool ExportTarget(const TSharedRef<SWindow>& ParentWindow, FLocalizationTargetSettings& TargetSettings);
+	bool ExportCulture(const TSharedRef<SWindow>& ParentWindow, FLocalizationTargetSettings& TargetSettings, const FString& CultureName);
+
+	bool GenerateReportsForTargets(const TSharedRef<SWindow>& ParentWindow, const TArray<FLocalizationTargetSettings*>& TargetsSettings);
+	bool GenerateReportsForTarget(const TSharedRef<SWindow>& ParentWindow, FLocalizationTargetSettings& TargetSettings);
+}
\ No newline at end of file
diff --git a/Engine/Source/Editor/LocalizationDashboard/Private/LocalizationConfigurationScript.cpp b/Engine/Source/Editor/LocalizationDashboard/Private/LocalizationConfigurationScript.cpp
new file mode 100644
--- /dev/null
+++ b/Engine/Source/Editor/LocalizationDashboard/Private/LocalizationConfigurationScript.cpp
@@ -0,0 +1,362 @@
+// Copyright 1998-2014 Epic Games, Inc. All Rights Reserved.
+
+#include "LocalizationDashboardPrivatePCH.h"
+#include "LocalizationConfigurationScript.h"
+#include "Classes/ProjectLocalizationSettings.h"
+
+namespace LocalizationConfigurationScript
+{
+	FString MakePathRelativeToProjectDirectory(const FString& Path)
+	{
+		FString Result = Path;
+		FPaths::MakePathRelativeTo(Result, *FPaths::GameDir());
+		return Result;
+	}
+
+	FString GetScriptDirectory()
+	{
+		return FPaths::GameConfigDir() / TEXT("Localization");
+	}
+
+	FString GetDataDirectory(const FLocalizationTargetSettings& Target)
+	{
+		return FPaths::GameContentDir() / TEXT("Localization") / Target.Name;
+	}
+
+	TArray<FString> GetScriptPaths(const FLocalizationTargetSettings& Target)
+	{
+		TArray<FString> Result;
+		Result.Add(GetGatherScriptPath(Target));
+		Result.Add(GetImportScriptPath(Target));
+		Result.Add(GetExportScriptPath(Target));
+		Result.Add(GetReportScriptPath(Target));
+		return Result;
+	}
+
+	FString GetManifestPath(const FLocalizationTargetSettings& Target)
+	{
+		return GetDataDirectory(Target) / FString::Printf( TEXT("%s.%s"), *Target.Name, TEXT("manifest") );
+	}
+
+	FString GetArchivePath(const FLocalizationTargetSettings& Target, const FString& CultureName)
+	{
+		return GetDataDirectory(Target) / CultureName / FString::Printf( TEXT("%s.%s"), *Target.Name, TEXT("archive") );
+	}
+
+	FString GetPOPath(const FLocalizationTargetSettings& Target, const FString& CultureName)
+	{
+		return GetDataDirectory(Target) / CultureName / FString::Printf( TEXT("%s.%s"), *Target.Name, TEXT("po") );
+	}
+
+	FString GetLocResPath(const FLocalizationTargetSettings& Target, const FString& CultureName)
+	{
+		return GetDataDirectory(Target) / CultureName / FString::Printf( TEXT("%s.%s"), *Target.Name, TEXT("locres") );
+	}
+
+	FLocalizationConfigurationScript GenerateGatherScript(const FLocalizationTargetSettings& Target)
+	{
+		FLocalizationConfigurationScript Script;
+
+		const bool HasSourceCode = IFileManager::Get().DirectoryExists( *FPaths::GameSourceDir() );
+
+		const FString ConfigDirRelativeToGameDir = MakePathRelativeToProjectDirectory(FPaths::GameConfigDir());
+		const FString SourceDirRelativeToGameDir = MakePathRelativeToProjectDirectory(FPaths::GameSourceDir());
+		const FString ContentDirRelativeToGameDir = MakePathRelativeToProjectDirectory(FPaths::GameContentDir());
+
+		// CommonSettings
+		{
+			FConfigSection& ConfigSection = Script.CommonSettings();
+
+			const UProjectLocalizationSettings* const ProjectLocalizationSettings = GetDefault<UProjectLocalizationSettings>(UProjectLocalizationSettings::StaticClass());
+			for (const FString& TargetDependencyName : Target.TargetDependencies)
+			{
+				ULocalizationTarget* const * OtherTarget = ProjectLocalizationSettings->TargetObjects.FindByPredicate([&TargetDependencyName](ULocalizationTarget* const OtherTarget)->bool{return OtherTarget->Settings.Name == TargetDependencyName;});
+				if (OtherTarget)
+				{
+					ConfigSection.Add( TEXT("ManifestDependencies"), MakePathRelativeToProjectDirectory(GetManifestPath((*OtherTarget)->Settings)) );
+				}
+			}
+
+			for (const FFilePath& Path : Target.AdditionalManifestDependencies)
+			{
+				ConfigSection.Add( TEXT("ManifestDependencies"), MakePathRelativeToProjectDirectory(Path.FilePath) );
+			}
+
+			const FString SourcePath = ContentDirRelativeToGameDir / TEXT("Localization") / Target.Name;
+			ConfigSection.Add( TEXT("SourcePath"), SourcePath );
+			const FString DestinationPath = ContentDirRelativeToGameDir / TEXT("Localization") / Target.Name;
+			ConfigSection.Add( TEXT("DestinationPath"), DestinationPath );
+
+			ConfigSection.Add( TEXT("ManifestName"), FPaths::GetCleanFilename(GetManifestPath(Target)) );
+			ConfigSection.Add( TEXT("ArchiveName"), FPaths::GetCleanFilename(GetArchivePath(Target, FString())) );
+
+			ConfigSection.Add( TEXT("SourceCulture"), Target.NativeCultureStatistics.CultureName );
+			ConfigSection.Add( TEXT("CulturesToGenerate"), Target.NativeCultureStatistics.CultureName );
+			for (const FCultureStatistics& CultureStatistics : Target.SupportedCulturesStatistics)
+			{
+				ConfigSection.Add( TEXT("CulturesToGenerate"), CultureStatistics.CultureName );
+			}
+		}
+
+		uint32 GatherTextStepIndex = 0;
+		// GatherTextFromSource
+		{
+			FConfigSection& ConfigSection = Script.GatherTextStep(GatherTextStepIndex++);
+
+			// CommandletClass
+			ConfigSection.Add( TEXT("CommandletClass"), TEXT("GatherTextFromSource") );
+
+			// Include Paths
+			if (HasSourceCode) 
+			{
+				ConfigSection.Add( TEXT("IncludePaths"), SourceDirRelativeToGameDir );
+			}
+			ConfigSection.Add( TEXT("IncludePaths"), ConfigDirRelativeToGameDir );
+
+			// Exclude Paths
+			ConfigSection.Add( TEXT("ExcludePaths"), ConfigDirRelativeToGameDir / TEXT("Localization") );
+
+			// SourceFileSearchFilters
+			if (HasSourceCode)
+			{
+				ConfigSection.Add( TEXT("SourceFileSearchFilters"), TEXT("*.h") );
+				ConfigSection.Add( TEXT("SourceFileSearchFilters"), TEXT("*.cpp") );
+			}
+			ConfigSection.Add( TEXT("SourceFileSearchFilters"), TEXT("*.ini") );
+		}
+
+		// GatherTextFromAssets
+		{
+			FConfigSection& ConfigSection = Script.GatherTextStep(GatherTextStepIndex++);
+
+			// CommandletClass
+			ConfigSection.Add( TEXT("CommandletClass"), TEXT("GatherTextFromAssets") );
+
+			// IncludePaths
+			ConfigSection.Add( TEXT("IncludePaths"), ContentDirRelativeToGameDir );
+
+			// ExcludePaths
+			ConfigSection.Add( TEXT("ExcludePaths"), ContentDirRelativeToGameDir / TEXT("Localization") );
+
+			// PackageExtensions
+			ConfigSection.Add( TEXT("PackageExtensions"), TEXT("*.umap") );
+			ConfigSection.Add( TEXT("PackageExtensions"), TEXT("*.uasset") );
+		}
+
+		// GenerateGatherManifest
+		{
+			FConfigSection& ConfigSection = Script.GatherTextStep(GatherTextStepIndex++);
+
+			// CommandletClass
+			ConfigSection.Add( TEXT("CommandletClass"), TEXT("GenerateGatherManifest") );
+		}
+
+		// GenerateGatherArchive
+		if (Target.SupportedCulturesStatistics.Num())
+		{
+			FConfigSection& ConfigSection = Script.GatherTextStep(GatherTextStepIndex++);
+
+			// CommandletClass
+			ConfigSection.Add( TEXT("CommandletClass"), TEXT("GenerateGatherArchive") );
+		}
+
+		Script.Dirty = true;
+
+		return Script;
+	}
+
+	FString GetGatherScriptPath(const FLocalizationTargetSettings& Target)
+	{
+		return GetScriptDirectory() / FString::Printf( TEXT("%s_Gather.%s"), *(Target.Name), TEXT("ini") );
+	}
+
+	FLocalizationConfigurationScript GenerateImportScript(const FLocalizationTargetSettings& Target, const FString* const SpecificCulture)
+	{
+		FLocalizationConfigurationScript Script;
+
+		const FString ContentDirRelativeToGameDir = MakePathRelativeToProjectDirectory(FPaths::GameContentDir());
+
+		// GatherTextStep0 - InternationalizationExport
+		{
+			FConfigSection& ConfigSection = Script.GatherTextStep(0);
+
+			// CommandletClass
+			ConfigSection.Add( TEXT("CommandletClass"), TEXT("InternationalizationExport") );
+
+			ConfigSection.Add( TEXT("bImportLoc"), TEXT("true") );
+
+			const FString SourcePath = ContentDirRelativeToGameDir / TEXT("Localization") / Target.Name;
+			ConfigSection.Add( TEXT("SourcePath"), SourcePath );
+			const FString DestinationPath = ContentDirRelativeToGameDir / TEXT("Localization") / Target.Name;
+			ConfigSection.Add( TEXT("DestinationPath"), DestinationPath );
+
+			const auto& AddCultureToGenerate = [&](const int32 Index)
+			{
+				ConfigSection.Add( TEXT("CulturesToGenerate"), Target.SupportedCulturesStatistics[Index].CultureName );
+			};
+
+			// Export for a specific culture.
+			if (SpecificCulture)
+			{
+				ConfigSection.Add( TEXT("CulturesToGenerate"), *SpecificCulture );
+			}
+			// Export for all cultures.
+			else
+			{
+				for (const FCultureStatistics& CultureStatistics : Target.SupportedCulturesStatistics)
+				{
+					ConfigSection.Add( TEXT("CulturesToGenerate"), CultureStatistics.CultureName );
+				}
+			}
+
+			ConfigSection.Add( TEXT("ManifestName"), FPaths::GetCleanFilename(GetManifestPath(Target)) );
+			ConfigSection.Add( TEXT("ArchiveName"), FPaths::GetCleanFilename(GetArchivePath(Target, FString())) );
+			ConfigSection.Add( TEXT("PortableObjectName"), FPaths::GetCleanFilename(GetPOPath(Target, FString())) );
+		}
+
+		Script.Dirty = true;
+
+		return Script;
+	}
+
+	FString GetImportScriptPath(const FLocalizationTargetSettings& Target, const FString* const SpecificCultureName)
+	{
+		const FString ConfigFileDirectory = GetScriptDirectory();
+		FString ConfigFilePath;
+		if (SpecificCultureName)
+		{
+			ConfigFilePath = ConfigFileDirectory / FString::Printf( TEXT("%s_Import_%s.%s"), *Target.Name, **SpecificCultureName, TEXT("ini") );
+		}
+		else
+		{
+			ConfigFilePath = ConfigFileDirectory / FString::Printf( TEXT("%s_Import.%s"), *Target.Name, TEXT("ini") );
+		}
+		return ConfigFilePath;
+	}
+
+	FLocalizationConfigurationScript GenerateExportScript(const FLocalizationTargetSettings& Target, const FString* const SpecificCulture)
+	{
+		FLocalizationConfigurationScript Script;
+
+		const FString ContentDirRelativeToGameDir = MakePathRelativeToProjectDirectory(FPaths::GameContentDir());
+
+		// GatherTextStep0 - InternationalizationExport
+		{
+			FConfigSection& ConfigSection = Script.GatherTextStep(0);
+
+			// CommandletClass
+			ConfigSection.Add( TEXT("CommandletClass"), TEXT("InternationalizationExport") );
+
+			ConfigSection.Add( TEXT("bExportLoc"), TEXT("true") );
+
+			const FString SourcePath = ContentDirRelativeToGameDir / TEXT("Localization") / Target.Name;
+			ConfigSection.Add( TEXT("SourcePath"), SourcePath );
+			const FString DestinationPath = ContentDirRelativeToGameDir / TEXT("Localization") / Target.Name;
+			ConfigSection.Add( TEXT("DestinationPath"), DestinationPath );
+
+			TArray<const FCultureStatistics*> AllCultureStatistics;
+			AllCultureStatistics.Add(&Target.NativeCultureStatistics);
+			for (const FCultureStatistics& SupportedCultureStatistics : Target.SupportedCulturesStatistics)
+			{
+				AllCultureStatistics.Add(&SupportedCultureStatistics);
+			}
+
+			const auto& AddCultureToGenerate = [&](const int32 Index)
+			{
+				ConfigSection.Add( TEXT("CulturesToGenerate"), AllCultureStatistics[Index]->CultureName );
+			};
+
+			// Export for a specific culture.
+			if (SpecificCulture)
+			{
+				const int32 CultureIndex = AllCultureStatistics.IndexOfByPredicate([SpecificCulture](const FCultureStatistics* Culture) { return Culture->CultureName == *SpecificCulture; });
+				AddCultureToGenerate(CultureIndex);
+			}
+			// Export for all cultures.
+			else
+			{
+				for (int32 CultureIndex = 0; CultureIndex < AllCultureStatistics.Num(); ++CultureIndex)
+				{
+					AddCultureToGenerate(CultureIndex);
+				}
+			}
+
+			ConfigSection.Add( TEXT("ManifestName"), FPaths::GetCleanFilename(GetManifestPath(Target)) );
+			ConfigSection.Add( TEXT("ArchiveName"), FPaths::GetCleanFilename(GetArchivePath(Target, FString())) );
+			ConfigSection.Add( TEXT("PortableObjectName"), FPaths::GetCleanFilename(GetPOPath(Target, FString())) );
+		}
+
+		Script.Dirty = true;
+
+		return Script;
+	}
+
+	FString GetExportScriptPath(const FLocalizationTargetSettings& Target, const FString* const SpecificCultureName)
+	{
+		const FString ConfigFileDirectory = GetScriptDirectory();
+		FString ConfigFilePath;
+		if (SpecificCultureName)
+		{
+			ConfigFilePath = ConfigFileDirectory / FString::Printf( TEXT("%s_Export_%s.%s"), *Target.Name, **SpecificCultureName, TEXT("ini") );
+		}
+		else
+		{
+			ConfigFilePath = ConfigFileDirectory / FString::Printf( TEXT("%s_Export.%s"), *Target.Name, TEXT("ini") );
+		}
+		return ConfigFilePath;
+	}
+
+	FLocalizationConfigurationScript GenerateReportScript(const FLocalizationTargetSettings& Target)
+	{
+		FLocalizationConfigurationScript Script;
+
+		const FString ContentDirRelativeToGameDir = MakePathRelativeToProjectDirectory(FPaths::GameContentDir());
+
+		// GatherTextStep0 - GenerateTextLocalizationReport
+		{
+			FConfigSection& ConfigSection = Script.GatherTextStep(0);
+
+			// CommandletClass
+			ConfigSection.Add( TEXT("CommandletClass"), TEXT("GenerateTextLocalizationReport") );
+
+			const FString SourcePath = ContentDirRelativeToGameDir / TEXT("Localization") / Target.Name;
+			ConfigSection.Add( TEXT("SourcePath"), SourcePath );
+			const FString DestinationPath = ContentDirRelativeToGameDir / TEXT("Localization") / Target.Name;
+			ConfigSection.Add( TEXT("DestinationPath"), DestinationPath );
+
+			ConfigSection.Add( TEXT("ManifestName"), FString::Printf( TEXT("%s.%s"), *Target.Name, TEXT("manifest") ) );
+
+			for (const FCultureStatistics& CultureStatistics : Target.SupportedCulturesStatistics)
+			{
+				ConfigSection.Add( TEXT("CulturesToGenerate"), CultureStatistics.CultureName );
+			}
+
+			ConfigSection.Add( TEXT("bWordCountReport"), TEXT("true") );
+
+			ConfigSection.Add( TEXT("WordCountReportName"), FPaths::GetCleanFilename( GetWordCountCSVPath(Target) ) );
+
+			ConfigSection.Add( TEXT("bConflictReport"), TEXT("true") );
+
+			ConfigSection.Add( TEXT("ConflictReportName"), FPaths::GetCleanFilename( GetConflictReportPath(Target) ) );
+		}
+
+		Script.Dirty = true;
+
+		return Script;
+	}
+
+	FString GetReportScriptPath(const FLocalizationTargetSettings& Target)
+	{
+		return GetScriptDirectory() / FString::Printf( TEXT("%s_GenerateReports.%s"), *(Target.Name), TEXT("ini") );
+	}
+
+	FString GetWordCountCSVPath(const FLocalizationTargetSettings& Target)
+	{
+		return GetDataDirectory(Target) / FString::Printf( TEXT("%s.%s"), *Target.Name, TEXT("csv") );
+	}
+
+	FString GetConflictReportPath(const FLocalizationTargetSettings& Target)
+	{
+		return GetDataDirectory(Target) / FString::Printf( TEXT("%s_Conflicts.%s"), *Target.Name, TEXT("txt") );
+	}
+}
\ No newline at end of file
diff --git a/Engine/Source/Editor/LocalizationDashboard/Private/LocalizationConfigurationScript.h b/Engine/Source/Editor/LocalizationDashboard/Private/LocalizationConfigurationScript.h
new file mode 100644
--- /dev/null
+++ b/Engine/Source/Editor/LocalizationDashboard/Private/LocalizationConfigurationScript.h
@@ -0,0 +1,47 @@
+// Copyright 1998-2014 Epic Games, Inc. All Rights Reserved.
+
+#pragma once
+
+#include "ConfigCacheIni.h"
+#include "Classes/LocalizationTarget.h"
+
+struct FLocalizationConfigurationScript : public FConfigFile
+{
+	FConfigSection& CommonSettings()
+	{
+		return FindOrAdd(TEXT("CommonSettings"));
+	}
+
+	FConfigSection& GatherTextStep(const uint32 Index)
+	{
+		return FindOrAdd( FString::Printf( TEXT("GatherTextStep%u"), Index) );
+	}
+};
+
+namespace LocalizationConfigurationScript
+{
+	FString MakePathRelativeToProjectDirectory(const FString& Path);
+
+	FString GetScriptDirectory();
+	FString GetDataDirectory(const FLocalizationTargetSettings& Target);
+	TArray<FString> GetScriptPaths(const FLocalizationTargetSettings& Target);
+
+	FString GetManifestPath(const FLocalizationTargetSettings& Target);
+	FString GetArchivePath(const FLocalizationTargetSettings& Target, const FString& CultureName);
+	FString GetPOPath(const FLocalizationTargetSettings& Target, const FString& CultureName);
+	FString GetLocResPath(const FLocalizationTargetSettings& Target, const FString& CultureName);
+
+	FLocalizationConfigurationScript GenerateGatherScript(const FLocalizationTargetSettings& Target);
+	FString GetGatherScriptPath(const FLocalizationTargetSettings& Target);
+
+	FLocalizationConfigurationScript GenerateImportScript(const FLocalizationTargetSettings& Target, const FString* const SpecificCulture = nullptr);
+	FString GetImportScriptPath(const FLocalizationTargetSettings& Target, const FString* const SpecificCultureName = nullptr);
+
+	FLocalizationConfigurationScript GenerateExportScript(const FLocalizationTargetSettings& Target, const FString* const SpecificCulture = nullptr);
+	FString GetExportScriptPath(const FLocalizationTargetSettings& Target, const FString* const SpecificCultureName = nullptr);
+
+	FLocalizationConfigurationScript GenerateReportScript(const FLocalizationTargetSettings& Target);
+	FString GetReportScriptPath(const FLocalizationTargetSettings& Target);
+	FString GetWordCountCSVPath(const FLocalizationTargetSettings& Target);
+	FString GetConflictReportPath(const FLocalizationTargetSettings& Target);
+}
\ No newline at end of file
diff --git a/Engine/Source/Editor/LocalizationDashboard/Private/LocalizationDashboardModule.cpp b/Engine/Source/Editor/LocalizationDashboard/Private/LocalizationDashboardModule.cpp
new file mode 100644
--- /dev/null
+++ b/Engine/Source/Editor/LocalizationDashboard/Private/LocalizationDashboardModule.cpp
@@ -0,0 +1,83 @@
+// Copyright 1998-2014 Epic Games, Inc. All Rights Reserved.
+
+#include "LocalizationDashboardPrivatePCH.h"
+#include "ILocalizationDashboardModule.h"
+#include "SLocalizationDashboard.h"
+#include "ProjectLocalizationSettings.h"
+#include "SDockTab.h"
+#include "Features/IModularFeatures.h"
+#include "ILocalizationServiceProvider.h"
+
+class FLocalizationDashboardModule
+	: public ILocalizationDashboardModule
+{
+public:
+	FLocalizationDashboardModule()
+		: Settings(GetMutableDefault<UProjectLocalizationSettings>(UProjectLocalizationSettings::StaticClass()))
+	{
+		
+		for (ULocalizationTarget*& Target : Settings->TargetObjects)
+		{
+			Target->Settings.UpdateStatusFromConflictReport();
+			Target->Settings.UpdateWordCountsFromCSV();
+		}
+	}
+
+	// Begin IModuleInterface interface
+	virtual void StartupModule() override
+	{
+		const auto& SpawnMainTab = [this](const FSpawnTabArgs& Args) -> TSharedRef<SDockTab>
+		{
+			return SNew(SDockTab)
+				.Label(NSLOCTEXT("LocalizationDashboard", "MainTabTitle", "Localization Dashboard"))
+				.TabRole(ETabRole::MajorTab)
+				[
+					SNew(SLocalizationDashboard, Settings.Get())
+				];
+		};
+
+		FGlobalTabmanager::Get()->RegisterTabSpawner(SLocalizationDashboard::TabName, FOnSpawnTab::CreateLambda( TFunction<TSharedRef<SDockTab> (const FSpawnTabArgs&)>(SpawnMainTab) ) )
+			.SetDisplayName(NSLOCTEXT("LocalizationDashboard", "MainTabTitle", "Localization Dashboard"))
+			.SetMenuType(ETabSpawnerMenuType::Hidden);
+
+		ServiceProviders = IModularFeatures::Get().GetModularFeatureImplementations<ILocalizationServiceProvider>("LocalizationService");
+		ServiceProviders.Insert(nullptr, 0); // "None"
+	}
+
+	virtual void ShutdownModule() override
+	{
+		FGlobalTabmanager::Get()->UnregisterTabSpawner(SLocalizationDashboard::TabName);
+	}
+	// End IModuleInterface interface
+
+	virtual void Show()
+	{
+		FGlobalTabmanager::Get()->InvokeTab(SLocalizationDashboard::TabName);
+	}
+
+	virtual const TArray<ILocalizationServiceProvider*>& GetLocalizationServiceProviders() const
+	{
+		return ServiceProviders;
+	}
+
+	virtual ILocalizationServiceProvider* GetCurrentLocalizationServiceProvider() const
+	{
+		if(Settings->LocalizationServiceProvider == NAME_None)
+		{
+			return nullptr;
+		}
+
+		auto ServiceProviderNameComparator = [&](ILocalizationServiceProvider* const LSP) -> bool
+		{
+			return LSP ? LSP->GetName() == Settings->LocalizationServiceProvider : false;
+		};
+		ILocalizationServiceProvider* const * LSP = ServiceProviders.FindByPredicate(ServiceProviderNameComparator);
+		return LSP ? *LSP : nullptr;
+	}
+
+private:
+	TWeakObjectPtr<UProjectLocalizationSettings> Settings;
+	TArray<ILocalizationServiceProvider*> ServiceProviders;
+};
+
+IMPLEMENT_MODULE(FLocalizationDashboardModule, "LocalizationDashboard");
\ No newline at end of file
diff --git a/Engine/Source/Editor/LocalizationDashboard/Private/LocalizationDashboardPrivatePCH.h b/Engine/Source/Editor/LocalizationDashboard/Private/LocalizationDashboardPrivatePCH.h
new file mode 100644
--- /dev/null
+++ b/Engine/Source/Editor/LocalizationDashboard/Private/LocalizationDashboardPrivatePCH.h
@@ -0,0 +1,8 @@
+// Copyright 1998-2014 Epic Games, Inc. All Rights Reserved.
+
+#pragma once
+
+/* Dependencies
+ *****************************************************************************/
+
+#include "SlateBasics.h"
\ No newline at end of file
diff --git a/Engine/Source/Editor/LocalizationDashboard/Private/LocalizationTarget.cpp b/Engine/Source/Editor/LocalizationDashboard/Private/LocalizationTarget.cpp
new file mode 100644
--- /dev/null
+++ b/Engine/Source/Editor/LocalizationDashboard/Private/LocalizationTarget.cpp
@@ -0,0 +1,450 @@
+// Copyright 1998-2014 Epic Games, Inc. All Rights Reserved.
+
+#include "LocalizationDashboardPrivatePCH.h"
+#include "Classes/LocalizationTarget.h"
+#include "Classes/ProjectLocalizationSettings.h"
+#include "ISourceControlModule.h"
+#include "CsvParser.h"
+#include "LocalizationConfigurationScript.h"
+
+FLocalizationTargetSettings::FLocalizationTargetSettings()
+	: Status(ELocalizationTargetStatus::Unknown)
+{
+}
+
+namespace
+{
+	class FWordCountCSVParser
+	{
+		typedef TArray<const TCHAR*> FRow;
+
+	public:
+		static bool Execute(FLocalizationTargetSettings& Target, const FString& CSVFilePath)
+		{
+			// Load CSV file to string.
+			FString CSVString;
+			FFileHelper::LoadFileToString(CSVString, *CSVFilePath);
+
+			// Parse CSV file string.
+			FCsvParser CSVParser(CSVString);
+			const FCsvParser::FRows& CSVRows = CSVParser.GetRows();
+
+			// Initialize.
+			FWordCountCSVParser Instance(Target);
+
+			// Parse header row.
+			if ( !(CSVRows.Num() > 0) ||!Instance.ParseHeaderRow(CSVRows[0]) )
+			{
+				return false;
+			}
+
+			// Parse latest word count row.
+			TMap<FString, uint32> WordCounts;
+			if ( !Instance.ParseWordCountRow(CSVRows.Top(), WordCounts) )
+			{
+				return false;
+			}
+
+			uint32* NativeWordCount = WordCounts.Find(Target.NativeCultureStatistics.CultureName);
+			if (NativeWordCount)
+			{
+				Target.NativeCultureStatistics.WordCount = *NativeWordCount;
+			}
+
+			for (FCultureStatistics& SupportedCultureStatistics : Target.SupportedCulturesStatistics)
+			{
+				uint32* SupportedWordCount = WordCounts.Find(SupportedCultureStatistics.CultureName);
+				if (SupportedWordCount)
+				{
+					SupportedCultureStatistics.WordCount = *SupportedWordCount;
+				}
+			}
+
+			return true;
+		}
+
+	private:
+		FWordCountCSVParser(FLocalizationTargetSettings& InTarget)
+			: Target(InTarget)
+		{
+			MandatoryColumnNames.Add( TEXT("Date/Time") );
+			MandatoryColumnNames.Add( TEXT("Word Count") );
+		}
+
+		bool ParseHeaderRow(const FRow& HeaderRow)
+		{
+			// Validate header row.
+			const bool IsHeaderRowValid = [&]() -> bool
+			{
+				if (HeaderRow.Num() < MandatoryColumnNames.Num())
+				{
+					// Not enough columns present to represent the mandatory columns.
+					return false;
+				}
+				for (int32 i = 0; i < MandatoryColumnNames.Num(); ++i)
+				{
+					if (MandatoryColumnNames[i] != HeaderRow[i])
+					{
+						return false;
+					}
+				}
+
+				return true;
+			}();
+
+			if (!IsHeaderRowValid)
+			{
+				return false;
+			}
+
+			// Parse culture names.
+			for (int32 i = MandatoryColumnNames.Num(); i < HeaderRow.Num(); ++i)
+			{
+				CultureNames.Add(HeaderRow[i]);
+			}
+
+			return true;
+		}
+
+		bool ParseWordCountRow(const FRow& WordCountRow, TMap<FString, uint32>& OutCultureWordCounts)
+		{
+			// Validate row.
+			bool IsRowValid = [&]() -> bool
+			{
+				if (WordCountRow.Num() != MandatoryColumnNames.Num() + CultureNames.Num())
+				{
+					return false;
+				}
+
+				return true;
+			}();
+
+			if (!IsRowValid)
+			{
+				return false;
+			}
+
+			// Parse time stamp.
+			FDateTime TimeStamp;
+			if (!FDateTime::Parse(WordCountRow[0], TimeStamp))
+			{
+				return false;
+			}
+
+			// Parse native culture word count.
+			{
+				const uint32 WordCount = static_cast<uint32>( FCString::Atoi(WordCountRow[1]) );
+				OutCultureWordCounts.Add(Target.NativeCultureStatistics.CultureName, WordCount);
+			}
+
+			// Parse supported culture word counts.
+			for (int32 j = MandatoryColumnNames.Num(); j < WordCountRow.Num(); ++j)
+			{
+				const uint32 WordCount = static_cast<uint32>( FCString::Atoi(WordCountRow[j]) );
+				const int32 CultureIndex = j - MandatoryColumnNames.Num();
+				OutCultureWordCounts.FindOrAdd(CultureNames[CultureIndex]) = WordCount;
+			}
+
+			return true;
+		}
+
+	private:
+		FLocalizationTargetSettings& Target;
+		TArray<FString> MandatoryColumnNames;
+		TArray<FString> CultureNames;
+	};
+}
+
+bool FLocalizationTargetSettings::UpdateWordCountsFromCSV()
+{
+	const FString CSVFilePath = LocalizationConfigurationScript::GetWordCountCSVPath(*this);
+	return FWordCountCSVParser::Execute(*this, CSVFilePath);
+}
+
+void FLocalizationTargetSettings::UpdateStatusFromConflictReport()
+{
+	const FString ConflictReportFilePath = LocalizationConfigurationScript::GetConflictReportPath(*this);
+	const int64 ConflictReportFileSize = IFileManager::Get().FileSize(*ConflictReportFilePath);
+
+	switch (ConflictReportFileSize)
+	{
+	case INDEX_NONE:
+		// Unknown
+		Status = ELocalizationTargetStatus::Unknown;
+		break;
+	case 0:
+		// Good
+		Status = ELocalizationTargetStatus::Clear;
+		break;
+	default:
+		// Bad
+		Status = ELocalizationTargetStatus::ConflictsPresent;
+		break;
+	}
+}
+
+bool FLocalizationTargetSettings::Rename(const FString& NewName)
+{
+	bool HasCompletelySucceeded = true;
+	ISourceControlModule& SourceControl = ISourceControlModule::Get();
+	ISourceControlProvider& SourceControlProvider = SourceControl.GetProvider();
+	const bool CanUseSourceControl = SourceControl.IsEnabled() && SourceControlProvider.IsEnabled() && SourceControlProvider.IsAvailable();
+
+	const auto& TryDelete = [&](const FString& Path) -> bool
+	{
+		const FSourceControlStatePtr SourceControlState = SourceControlProvider.GetState(*Path, EStateCacheUsage::Use);
+
+		// Should only use source control if the file is actually under source control.
+		if (CanUseSourceControl && SourceControlState.IsValid() && !SourceControlState->CanAdd())
+		{
+			// File is already marked for deletion.
+			if (SourceControlState->IsDeleted())
+			{
+				return true;
+			}
+			else 
+			{
+				// File is in some modified source control state, needs to be reverted.
+				if (SourceControlState->IsAdded() || SourceControlState->IsCheckedOut())
+				{
+					if (SourceControlProvider.Execute(ISourceControlOperation::Create<FRevert>(), *Path) != ECommandResult::Succeeded)
+					{
+						return false;
+					}
+				}
+
+				// File needs to be deleted.
+				if (SourceControlProvider.Execute(ISourceControlOperation::Create<FDelete>(), *Path) == ECommandResult::Succeeded)
+				{
+					return true;
+				}
+			}
+		}
+		// Attempt local deletion.
+		else if (IFileManager::Get().Delete(*Path, false, true))
+		{
+			return true;
+		}
+
+		return false;
+	};
+
+	// Delete configuration files.
+	const TArray<FString>& ScriptPaths = LocalizationConfigurationScript::GetScriptPaths(*this);
+	for (const FString& ScriptPath : ScriptPaths)
+	{
+		if (!TryDelete(ScriptPath))
+		{
+			HasCompletelySucceeded = false;
+		}
+	}
+
+	const auto& GetNamedPaths = [this](TArray<FString>& NamedPaths)
+	{
+		NamedPaths.Add(LocalizationConfigurationScript::GetManifestPath(*this));
+		NamedPaths.Add(LocalizationConfigurationScript::GetWordCountCSVPath(*this));
+		NamedPaths.Add(LocalizationConfigurationScript::GetConflictReportPath(*this));
+
+		TArray<const FCultureStatistics*> Cultures;
+		Cultures.Add(&NativeCultureStatistics);
+		for (const FCultureStatistics& Culture : SupportedCulturesStatistics)
+		{
+			Cultures.Add(&Culture);
+		}
+
+		for (const FCultureStatistics* Culture : Cultures)
+		{
+			NamedPaths.Add(LocalizationConfigurationScript::GetArchivePath(*this, Culture->CultureName));
+			NamedPaths.Add(LocalizationConfigurationScript::GetPOPath(*this, Culture->CultureName));
+			NamedPaths.Add(LocalizationConfigurationScript::GetLocResPath(*this, Culture->CultureName));
+		}
+
+		NamedPaths.Add(LocalizationConfigurationScript::GetDataDirectory(*this));
+	};
+
+	TArray<FString> OldPaths;
+	GetNamedPaths(OldPaths);
+
+	// Rename
+	Name = NewName;
+
+	TArray<FString> NewPaths;
+	GetNamedPaths(NewPaths);
+
+	// Rename data files.
+	check(OldPaths.Num() == NewPaths.Num());
+	for (int32 i = 0; i < OldPaths.Num(); ++i)
+	{
+		if (IFileManager::Get().Move(*NewPaths[i], *OldPaths[i], true, true, false, true))
+		{
+			if (CanUseSourceControl)
+			{
+				// Add new file.
+				if (SourceControlProvider.Execute(ISourceControlOperation::Create<FMarkForAdd>(), *NewPaths[i]) != ECommandResult::Succeeded)
+				{
+					HasCompletelySucceeded = false;
+				}
+			}
+
+			bool ShouldUseSourceControl = false;
+
+			if (CanUseSourceControl)
+			{
+				const FSourceControlStatePtr SourceControlState = SourceControlProvider.GetState(*OldPaths[i], EStateCacheUsage::Use);
+				if (SourceControlState->IsSourceControlled())
+				{
+					ShouldUseSourceControl = true;
+				}
+			}
+
+			// Use source control.
+			if (ShouldUseSourceControl)
+			{
+				// Delete old directory/file.
+				if (SourceControlProvider.Execute(ISourceControlOperation::Create<FDelete>(), *OldPaths[i]) != ECommandResult::Succeeded)
+				{
+					HasCompletelySucceeded = false;
+				}
+			}
+			// Operate locally.
+			else
+			{
+				// Delete old directory/file.
+				if (	(FPaths::DirectoryExists(OldPaths[i]) && !IFileManager::Get().DeleteDirectory(*OldPaths[i], false, true)) ||
+					(FPaths::FileExists(OldPaths[i]) && !IFileManager::Get().Delete(*OldPaths[i], false, true))	)
+				{
+					HasCompletelySucceeded = false;
+				}
+			}
+		}
+	}
+
+	// Generate new configuration files.
+	LocalizationConfigurationScript::GenerateGatherScript(*this).Write(LocalizationConfigurationScript::GetGatherScriptPath(*this));
+	LocalizationConfigurationScript::GenerateImportScript(*this).Write(LocalizationConfigurationScript::GetImportScriptPath(*this));
+	LocalizationConfigurationScript::GenerateExportScript(*this).Write(LocalizationConfigurationScript::GetExportScriptPath(*this));
+	LocalizationConfigurationScript::GenerateReportScript(*this).Write(LocalizationConfigurationScript::GetReportScriptPath(*this));
+
+	//if (CanUseSourceControl)
+	//{
+	//	HasCompletelySucceeded = HasCompletelySucceeded && SourceControlProvider.Execute(ISourceControlOperation::Create<FMarkForAdd>(), *LocalizationConfigurationScript::GetGatherScriptPath(*this)) == ECommandResult::Succeeded;
+	//	HasCompletelySucceeded = HasCompletelySucceeded && SourceControlProvider.Execute(ISourceControlOperation::Create<FMarkForAdd>(), *LocalizationConfigurationScript::GetImportScriptPath(*this)) == ECommandResult::Succeeded;
+	//	HasCompletelySucceeded = HasCompletelySucceeded && SourceControlProvider.Execute(ISourceControlOperation::Create<FMarkForAdd>(), *LocalizationConfigurationScript::GetExportScriptPath(*this)) == ECommandResult::Succeeded;
+	//	HasCompletelySucceeded = HasCompletelySucceeded && SourceControlProvider.Execute(ISourceControlOperation::Create<FMarkForAdd>(), *LocalizationConfigurationScript::GetReportScriptPath(*this)) == ECommandResult::Succeeded;
+	//}
+
+	return HasCompletelySucceeded;
+}
+
+bool FLocalizationTargetSettings::DeleteFiles(const FString* const Culture) const
+{
+	bool HasCompletelySucceeded = true;
+	ISourceControlModule& SourceControl = ISourceControlModule::Get();
+	ISourceControlProvider& SourceControlProvider = SourceControl.GetProvider();
+
+	const FString DataDirectory = LocalizationConfigurationScript::GetDataDirectory(*this) + (Culture ? *Culture : FString());
+
+	// Remove data files from source control.
+	{
+		TArray<FString> FilesInDataDirectory;
+		IFileManager::Get().FindFilesRecursive(FilesInDataDirectory, *DataDirectory, TEXT("*"), true, false);
+
+		ISourceControlProvider& SourceControlProvider = ISourceControlModule::Get().GetProvider();
+		TArray<FSourceControlStateRef> SourceControlStates;
+		ECommandResult::Type Result = SourceControlProvider.GetState( FilesInDataDirectory, SourceControlStates, EStateCacheUsage::ForceUpdate );
+		if (Result == ECommandResult::Succeeded)
+		{
+			TArray<FString> FilesToRevert;
+			TArray<FString> FilesToDelete;
+
+			for (const FSourceControlStateRef& SourceControlState : SourceControlStates)
+			{
+				// Added/edited files must be reverted.
+				if (SourceControlState->IsAdded() || SourceControlState->IsCheckedOut())
+				{
+					FilesToRevert.Add(SourceControlState->GetFilename());
+				}
+				// Edited files must be deleted.
+				if (SourceControlState->IsCheckedOut())
+				{
+					FilesToDelete.Add(SourceControlState->GetFilename());
+				}
+			}
+
+			const bool WasSuccessfullyReverted = (SourceControlProvider.Execute( ISourceControlOperation::Create<FRevert>(), FilesToRevert ) == ECommandResult::Succeeded);
+			const bool WasSuccessfullyDeleted = (SourceControlProvider.Execute( ISourceControlOperation::Create<FDelete>(), FilesToDelete ) == ECommandResult::Succeeded);
+			HasCompletelySucceeded = HasCompletelySucceeded && WasSuccessfullyReverted && WasSuccessfullyDeleted;
+		}
+	}
+
+	// Delete data files.
+	if (!IFileManager::Get().DeleteDirectory(*DataDirectory, false, true))
+	{
+		HasCompletelySucceeded = false;
+	}
+
+	// Delete configuration files if deleting the target.
+	if (!Culture)
+	{
+		const TArray<FString>& ScriptPaths = LocalizationConfigurationScript::GetScriptPaths(*this);
+
+		// Remove script files from source control.
+		{
+			ISourceControlProvider& SourceControlProvider = ISourceControlModule::Get().GetProvider();
+			TArray<FSourceControlStateRef> SourceControlStates;
+			ECommandResult::Type Result = SourceControlProvider.GetState( ScriptPaths, SourceControlStates, EStateCacheUsage::ForceUpdate );
+			if (Result == ECommandResult::Succeeded)
+			{
+				TArray<FString> FilesToRevert;
+				TArray<FString> FilesToDelete;
+
+				for (const FSourceControlStateRef& SourceControlState : SourceControlStates)
+				{
+					// Added/edited files must be reverted.
+					if (SourceControlState->IsAdded() || SourceControlState->IsCheckedOut())
+					{
+						FilesToRevert.Add(SourceControlState->GetFilename());
+					}
+					// Edited files must be deleted.
+					if (SourceControlState->IsCheckedOut())
+					{
+						FilesToDelete.Add(SourceControlState->GetFilename());
+					}
+				}
+
+				const bool WasSuccessfullyReverted = (SourceControlProvider.Execute( ISourceControlOperation::Create<FRevert>(), FilesToRevert ) == ECommandResult::Succeeded);
+				const bool WasSuccessfullyDeleted = (SourceControlProvider.Execute( ISourceControlOperation::Create<FDelete>(), FilesToDelete ) == ECommandResult::Succeeded);
+				HasCompletelySucceeded = HasCompletelySucceeded && WasSuccessfullyReverted && WasSuccessfullyDeleted;
+			}
+		}
+
+		// Delete script files.
+		for (const FString& ScriptPath : ScriptPaths)
+		{
+			if (!IFileManager::Get().Delete(*ScriptPath, false, true))
+			{
+				HasCompletelySucceeded = false;
+			}
+		}
+	}
+
+	return HasCompletelySucceeded;
+}
+
+ULocalizationTarget::ULocalizationTarget(const FObjectInitializer& OI)
+	: Super(OI)
+{
+}
+
+void ULocalizationTarget::PostEditChangeProperty(FPropertyChangedEvent& PropertyChangedEvent)
+{
+	UProjectLocalizationSettings* const ParentProjectLocalizationSettings = Cast<UProjectLocalizationSettings>(GetOuter());
+	if (ParentProjectLocalizationSettings)
+	{
+		UClass* ParentClass = ParentProjectLocalizationSettings ? ParentProjectLocalizationSettings->GetClass() : nullptr;
+		UProperty* ChangedParentProperty = ParentClass ? ParentClass->FindPropertyByName(GET_MEMBER_NAME_CHECKED(UProjectLocalizationSettings,TargetObjects)) : nullptr;
+
+		FPropertyChangedEvent PropertyChangedEventForParent(ChangedParentProperty);
+		ParentProjectLocalizationSettings->PostEditChangeProperty(PropertyChangedEventForParent);
+	}
+}
\ No newline at end of file
diff --git a/Engine/Source/Editor/LocalizationDashboard/Private/ProjectLocalizationSettings.cpp b/Engine/Source/Editor/LocalizationDashboard/Private/ProjectLocalizationSettings.cpp
new file mode 100644
--- /dev/null
+++ b/Engine/Source/Editor/LocalizationDashboard/Private/ProjectLocalizationSettings.cpp
@@ -0,0 +1,58 @@
+// Copyright 1998-2014 Epic Games, Inc. All Rights Reserved.
+
+#include "LocalizationDashboardPrivatePCH.h"
+#include "Classes/ProjectLocalizationSettings.h"
+#include "ISourceControlModule.h"
+
+UProjectLocalizationSettings::UProjectLocalizationSettings(const FObjectInitializer& OI)
+	: Super(OI)
+{
+}
+
+void UProjectLocalizationSettings::PostInitProperties()
+{
+	Super::PostInitProperties();
+
+	TargetObjects.Empty(Targets.Num());
+	for (const FLocalizationTargetSettings& Target : Targets)
+	{
+		ULocalizationTarget* const NewTargetObject = Cast<ULocalizationTarget>(StaticConstructObject(ULocalizationTarget::StaticClass(), this));
+		NewTargetObject->Settings = Target;
+		TargetObjects.Add(NewTargetObject);
+	}
+}
+
+void UProjectLocalizationSettings::PostEditChangeProperty(FPropertyChangedEvent& PropertyChangedEvent)
+{
+	Super::PostEditChangeProperty(PropertyChangedEvent);
+
+	if (PropertyChangedEvent.MemberProperty->GetFName() == GET_MEMBER_NAME_CHECKED(UProjectLocalizationSettings,TargetObjects))
+	{
+		Targets.Empty(TargetObjects.Num());
+		for (ULocalizationTarget* const& TargetObject : TargetObjects)
+		{
+			Targets.Add(TargetObject ? TargetObject->Settings : FLocalizationTargetSettings());
+		}
+	}
+
+	ISourceControlModule& SourceControl = ISourceControlModule::Get();
+	ISourceControlProvider& SourceControlProvider = SourceControl.GetProvider();
+	const bool CanUseSourceControl = SourceControl.IsEnabled() && SourceControlProvider.IsEnabled() && SourceControlProvider.IsAvailable();
+
+	if (CanUseSourceControl)
+	{
+		const FString Path = GetConfigFilename(this);
+		const FSourceControlStatePtr SourceControlState = SourceControlProvider.GetState(*Path, EStateCacheUsage::Use);
+
+		// Should only use source control if the file is actually under source control.
+		if (CanUseSourceControl && SourceControlState.IsValid() && !SourceControlState->CanAdd())
+		{
+			if (!SourceControlState->CanEdit())
+			{
+				SourceControlProvider.Execute(ISourceControlOperation::Create<FCheckOut>(), *Path);
+			}
+		}
+	}
+
+	SaveConfig();
+}
\ No newline at end of file
diff --git a/Engine/Source/Editor/LocalizationDashboard/Private/SCulturePicker.cpp b/Engine/Source/Editor/LocalizationDashboard/Private/SCulturePicker.cpp
new file mode 100644
--- /dev/null
+++ b/Engine/Source/Editor/LocalizationDashboard/Private/SCulturePicker.cpp
@@ -0,0 +1,271 @@
+// Copyright 1998-2014 Epic Games, Inc. All Rights Reserved.
+
+#include "LocalizationDashboardPrivatePCH.h"
+#include "SCulturePicker.h"
+#include "STableRow.h"
+#include "SSearchBox.h"
+
+void SCulturePicker::Construct( const FArguments& InArgs )
+{
+	OnCultureSelectionChanged = InArgs._OnSelectionChanged;
+	IsCulturePickable = InArgs._IsCulturePickable;
+
+	TArray<FString> AllCultureNames;
+	FInternationalization::Get().GetCultureNames(AllCultureNames);
+	for (const auto& CultureName : AllCultureNames)
+	{
+		const FCulturePtr Culture = FInternationalization::Get().GetCulture(CultureName);
+		if (Culture.IsValid())
+		{
+			Cultures.AddUnique(Culture);
+			const TArray<FString> ParentCultureNames = Culture->GetPrioritizedParentCultureNames();
+			for (const auto& ParentCultureName : ParentCultureNames)
+			{
+				const FCulturePtr ParentCulture = FInternationalization::Get().GetCulture(ParentCultureName);
+				Cultures.AddUnique(ParentCulture);
+			}
+		}
+	}
+
+	BuildStockEntries();
+	RebuildEntries();
+
+	ChildSlot
+		[
+			SNew(SVerticalBox)
+			+SVerticalBox::Slot()
+			.AutoHeight()
+			[
+				SNew(SSearchBox)
+				.HintText( NSLOCTEXT("CulturePicker", "SearchHintText", "Name/Abbreviation") )
+				.OnTextChanged(this, &SCulturePicker::OnFilterStringChanged)
+				.DelayChangeNotificationsWhileTyping(true)
+			]
+			+SVerticalBox::Slot()
+			.FillHeight(1.0f)
+			[
+				SAssignNew(TreeView, STreeView<TSharedPtr<FCultureEntry>>)
+				.SelectionMode(ESelectionMode::Single)
+				.TreeItemsSource(&RootEntries)
+				.OnGenerateRow(this, &SCulturePicker::OnGenerateRow)
+				.OnGetChildren(this, &SCulturePicker::OnGetChildren)
+				.OnSelectionChanged(this, &SCulturePicker::OnSelectionChanged)
+			]
+		];
+
+
+	const auto& IsInitialSelection = [&](const TSharedPtr<FCultureEntry>& Entry) -> bool
+	{
+		return Entry->Culture == InArgs._InitialSelection;
+	};
+
+	const TSharedPtr<FCultureEntry>* InitialSelection = RootEntries.FindByPredicate(IsInitialSelection);
+	if (InitialSelection)
+	{
+		TreeView->SetSelection(*InitialSelection);
+	}
+}
+
+void SCulturePicker::RequestTreeRefresh()
+{
+	RebuildEntries();
+	if (TreeView.IsValid())
+	{
+		TreeView->RequestTreeRefresh();
+	}
+}
+
+void SCulturePicker::BuildStockEntries()
+{
+	TArray<FCulturePtr> StockCultures;
+	TArray<FString> StockCultureNames;
+	FInternationalization::Get().GetCultureNames(StockCultureNames);
+	for (const auto& CultureName : StockCultureNames)
+	{
+		const FCulturePtr Culture = FInternationalization::Get().GetCulture(CultureName);
+		if (Culture.IsValid())
+		{
+			StockCultures.AddUnique(Culture);
+			const TArray<FString> ParentCultureNames = Culture->GetPrioritizedParentCultureNames();
+			for (const auto& ParentCultureName : ParentCultureNames)
+			{
+				const FCulturePtr ParentCulture = FInternationalization::Get().GetCulture(ParentCultureName);
+				StockCultures.AddUnique(ParentCulture);
+			}
+		}
+	}
+
+	TMap< FString, TSharedPtr<FCultureEntry> > StockCultureEntries;
+	StockCultureEntries.Reserve(StockCultures.Num());
+
+	StockEntries.Empty(StockCultures.Num());
+	StockEntries.Reserve(StockCultureEntries.Num());
+	for (const FCulturePtr& Culture : StockCultures)
+	{
+		const FString Name = Culture->GetName();
+
+		TSharedRef<FCultureEntry> CultureEntry = MakeShareable( new FCultureEntry(Culture) );
+
+		// Add culture to entries.
+		StockCultureEntries.Add( Name, CultureEntry );
+
+		// Add entry to list of root stock entries - it will be removed if someone has it as a child.
+		StockEntries.Add( CultureEntry );
+	}
+
+	// Update all parent entries to know about their child entries.
+	for (auto& Pair : StockCultureEntries)
+	{
+		const TSharedPtr<FCultureEntry> Entry = Pair.Value;
+
+		const TArray<FString> ParentCultureNames = Entry->Culture->GetPrioritizedParentCultureNames();
+
+		TArray<FString> InvalidAncestors;
+		for (const FString& ParentName : ParentCultureNames)
+		{
+			// Find existing parent entry.
+			TSharedPtr<FCultureEntry> ParentEntry;
+			{
+				const TSharedPtr<FCultureEntry>* const PossibleParentEntry = StockCultureEntries.Find(ParentName);
+				if (PossibleParentEntry)
+				{
+					ParentEntry = *PossibleParentEntry;
+				}
+			}
+
+			// Skip unfound parents.
+			if (!ParentEntry.IsValid())
+			{
+				continue;
+			}
+
+			// Don't add self as child of self.
+			if (Entry == ParentEntry)
+			{
+				continue;
+			}
+
+			// Invalidate ancestors of this parent entry.
+			TArray<FString> ParentAncestorCultureNames = ParentEntry->Culture->GetPrioritizedParentCultureNames();
+			for (const auto& AncestorCultureName : ParentAncestorCultureNames)
+			{
+				// Don't treat this parent entry as a parent of itself.
+				if (AncestorCultureName != ParentName)
+				{
+					InvalidAncestors.Add(AncestorCultureName);
+				}
+			}
+
+			// Don't add ancestors of ancestors.
+			if (InvalidAncestors.Contains(ParentName))
+			{
+				continue;
+			}
+
+			// Add current entry to parent entry.
+			ParentEntry->Children.Add( MakeShareable( new FCultureEntry(*Entry) ) );
+
+			// Remove the current entry from the base entries, as it is a child.
+			StockEntries.Remove(Entry);
+		}
+	}
+
+	// Sort entries.
+	auto CultureEntryComparator = [](const TSharedPtr<FCultureEntry>& LHS, const TSharedPtr<FCultureEntry>& RHS) -> bool
+	{
+		return LHS->Culture->GetDisplayName() < RHS->Culture->GetDisplayName();
+	};
+	StockEntries.Sort(CultureEntryComparator);
+}
+
+void SCulturePicker::RebuildEntries()
+{
+	RootEntries.Reset();
+
+	TFunctionRef<void (const TArray< TSharedPtr<FCultureEntry> >&, TArray< TSharedPtr<FCultureEntry> >&)> DeepCopyAndFilter
+		= [&](const TArray< TSharedPtr<FCultureEntry> >& InEntries, TArray< TSharedPtr<FCultureEntry> >& OutEntries)
+	{
+		for (const TSharedPtr<FCultureEntry>& InEntry : InEntries)
+		{
+			// Set pickable flag.
+			bool IsPickable = true;
+			if (IsCulturePickable.IsBound())
+			{
+				IsPickable = IsCulturePickable.Execute(InEntry->Culture);
+			}
+
+			TSharedRef<FCultureEntry> OutEntry = MakeShareable(new FCultureEntry(InEntry->Culture, IsPickable));
+
+			// Recurse to children.
+			DeepCopyAndFilter(InEntry->Children, OutEntry->Children);
+
+			bool IsFilteredOut = false;
+			if (!FilterString.IsEmpty())
+			{
+				const FString Name = OutEntry->Culture->GetName();
+				const FString DisplayName = OutEntry->Culture->GetDisplayName();
+				IsFilteredOut = !Name.Contains(FilterString) && !DisplayName.Contains(FilterString);
+			}
+
+			// If has children, must be added. If it is not filtered and it is pickable, should be added.
+			if ( OutEntry->Children.Num() != 0 || (!IsFilteredOut && IsPickable) )
+			{
+				OutEntries.Add(OutEntry);
+			}
+		}	
+	};
+	DeepCopyAndFilter(StockEntries, RootEntries);
+}
+
+void SCulturePicker::OnFilterStringChanged(const FText& InFilterString)
+{
+	FilterString = InFilterString.ToString();
+	RequestTreeRefresh();
+}
+
+TSharedRef<ITableRow> SCulturePicker::OnGenerateRow(TSharedPtr<FCultureEntry> Entry, const TSharedRef<STableViewBase>& Table)
+{
+	return	SNew(STableRow< TSharedPtr<FCultureEntry> >, Table)
+			[
+				SNew(STextBlock)
+				.Text( FText::FromString(Entry->Culture->GetDisplayName()) )
+				.ToolTip(
+						SNew(SToolTip)
+						.Content()
+						[
+							SNew(STextBlock)
+							.Text( FText::FromString(Entry->Culture->GetName()) )
+							.HighlightText( FText::FromString(FilterString) )
+						]
+						)
+				.HighlightText( FText::FromString(FilterString) )
+				
+				.ColorAndOpacity( Entry->IsSelectable ? FSlateColor::UseForeground() : FSlateColor::UseSubduedForeground() )
+				// TODO: Appear differently if not selectable but present because it's part of the hierarchy.
+			];
+}
+
+void SCulturePicker::OnGetChildren(TSharedPtr<FCultureEntry> Entry, TArray< TSharedPtr<FCultureEntry> >& Children)
+{
+	// Add entries from children array.
+	for (const auto& Child : Entry->Children)
+	{
+		Children.Add(Child);
+
+		// Sort entries.
+		auto CultureEntryComparator = [](const TSharedPtr<FCultureEntry>& LHS, const TSharedPtr<FCultureEntry>& RHS) -> bool
+		{
+			return LHS->Culture->GetDisplayName() < RHS->Culture->GetDisplayName();
+		};
+		Children.Sort(CultureEntryComparator);
+	}
+}
+
+void SCulturePicker::OnSelectionChanged(TSharedPtr<FCultureEntry> Entry, ESelectInfo::Type SelectInfo)
+{
+	// Don't count as selection if the entry isn't actually selectable but is part of the hierarchy.
+	if (Entry.IsValid() && Entry->IsSelectable)
+	{
+		OnCultureSelectionChanged.ExecuteIfBound( Entry->Culture, SelectInfo );
+	}
+}
\ No newline at end of file
diff --git a/Engine/Source/Editor/LocalizationDashboard/Private/SCulturePicker.h b/Engine/Source/Editor/LocalizationDashboard/Private/SCulturePicker.h
new file mode 100644
--- /dev/null
+++ b/Engine/Source/Editor/LocalizationDashboard/Private/SCulturePicker.h
@@ -0,0 +1,78 @@
+// Copyright 1998-2014 Epic Games, Inc. All Rights Reserved.
+
+#pragma once
+
+#include "SCompoundWidget.h"
+#include "SlateDelegates.h"
+#include "CulturePointer.h"
+
+struct FCultureEntry
+{
+	FCulturePtr Culture;
+	TArray< TSharedPtr<FCultureEntry> > Children;
+	bool IsSelectable;
+
+	FCultureEntry(const FCulturePtr& InCulture, const bool InIsSelectable = true)
+		: Culture(InCulture)
+		, IsSelectable(InIsSelectable)
+	{}
+
+	FCultureEntry(const FCultureEntry& Source)
+		: Culture(Source.Culture)
+		, IsSelectable(Source.IsSelectable)
+	{
+		Children.Reserve(Source.Children.Num());
+		for (const auto& Child : Source.Children)
+		{
+			Children.Add( MakeShareable( new FCultureEntry(*Child) ) );
+		}
+	}
+};
+
+class SCulturePicker : public SCompoundWidget
+{
+public:
+	/** A delegate type invoked when a selection changes somewhere. */
+	DECLARE_DELEGATE_RetVal_OneParam(bool, FIsCulturePickable, FCulturePtr);
+	typedef TSlateDelegates< FCulturePtr >::FOnSelectionChanged FOnSelectionChanged;
+
+public:
+	SLATE_BEGIN_ARGS( SCulturePicker ){}
+		SLATE_EVENT( FOnSelectionChanged, OnSelectionChanged )
+		SLATE_EVENT( FIsCulturePickable, IsCulturePickable )
+		SLATE_ARGUMENT( FCulturePtr, InitialSelection )
+	SLATE_END_ARGS()
+
+	void Construct( const FArguments& InArgs );
+	void RequestTreeRefresh();
+
+private:
+	void BuildStockEntries();
+	void RebuildEntries();
+
+	void OnFilterStringChanged(const FText& InFilterString);
+	TSharedRef<ITableRow> OnGenerateRow(TSharedPtr<FCultureEntry> Entry, const TSharedRef<STableViewBase>& Table);
+	void OnGetChildren(TSharedPtr<FCultureEntry> Entry, TArray< TSharedPtr<FCultureEntry> >& Children);
+	void OnSelectionChanged(TSharedPtr<FCultureEntry> Entry, ESelectInfo::Type SelectInfo);
+
+private:
+	TSharedPtr< STreeView< TSharedPtr<FCultureEntry> > > TreeView;
+
+	/* The provided cultures array. */
+	TArray<FCulturePtr> Cultures;
+
+	/* The top level culture entries for all possible stock cultures. */
+	TArray< TSharedPtr<FCultureEntry> > StockEntries;
+
+	/* The top level culture entries to be displayed in the tree view. */
+	TArray< TSharedPtr<FCultureEntry> > RootEntries;
+
+	/* The string by which to filter cultures names. */
+	FString FilterString;
+
+	/** Delegate to invoke when selection changes. */
+	FOnSelectionChanged OnCultureSelectionChanged;
+
+	/** Delegate to invoke to set whether a culture is "pickable". */
+	FIsCulturePickable IsCulturePickable;
+};
\ No newline at end of file
diff --git a/Engine/Source/Editor/LocalizationDashboard/Private/SLocalizationDashboard.cpp b/Engine/Source/Editor/LocalizationDashboard/Private/SLocalizationDashboard.cpp
new file mode 100644
--- /dev/null
+++ b/Engine/Source/Editor/LocalizationDashboard/Private/SLocalizationDashboard.cpp
@@ -0,0 +1,498 @@
+// Copyright 1998-2014 Epic Games, Inc. All Rights Reserved.
+
+#include "LocalizationDashboardPrivatePCH.h"
+#include "SLocalizationDashboard.h"
+#include "IDetailCustomization.h"
+#include "EditorStyleSet.h"
+#include "DetailLayoutBuilder.h"
+#include "ProjectLocalizationSettings.h"
+#include "PropertyEditorModule.h"
+#include "IDetailsView.h"
+#include "DetailWidgetRow.h"
+#include "DetailCategoryBuilder.h"
+#include "PropertyHandle.h"
+#include "SLocalizationDashboardTargetRow.h"
+#include "TargetsTableEntry.h"
+#include "Runtime/Core/Public/Features/IModularFeatures.h"
+#include "ILocalizationDashboardModule.h"
+#include "ILocalizationServiceProvider.h"
+#include "IPropertyUtilities.h"
+#include "MultiBoxBuilder.h"
+#include "FeedbackContext.h"
+#include "LocalizationCommandletTasks.h"
+
+#define LOCTEXT_NAMSPACE "LocalizationDashboard"
+
+const FName SLocalizationDashboard::TabName("LocalizationDashboard");
+
+class ILocalizationServiceProvider;
+
+namespace
+{
+	struct FLocalizationServiceProviderWrapper;
+
+	class FProjectLocalizationSettingsDetailsCustomization : public IDetailCustomization
+	{
+	public:
+		FProjectLocalizationSettingsDetailsCustomization(UProjectLocalizationSettings* const InSettings);
+		void CustomizeDetails(IDetailLayoutBuilder& DetailBuilder) override;
+
+	private:
+		void BuildTargetsList();
+		void RebuildTargetsList();
+
+		FText GetCurrentServiceProviderDisplayName() const;
+		TSharedRef<SWidget> ServiceProviderComboBox_OnGenerateWidget(TSharedPtr<FLocalizationServiceProviderWrapper> LSPWrapper) const;
+		void ServiceProviderComboBox_OnSelectionChanged(TSharedPtr<FLocalizationServiceProviderWrapper> LSPWrapper, ESelectInfo::Type SelectInfo);
+
+		void GatherAllTargets();
+		void ImportAllTargets();
+		void ExportAllTargets();
+		void UpdateTargetFromReports(ULocalizationTarget* const LocalizationTarget);
+
+		TSharedRef<ITableRow> OnGenerateRow(TSharedPtr<IPropertyHandle> TargetObjectPropertyHandle, const TSharedRef<STableViewBase>& Table);
+		FReply OnNewTargetButtonClicked();
+
+	private:
+		UProjectLocalizationSettings* Settings;
+
+		IDetailLayoutBuilder* DetailLayoutBuilder;
+
+		TSharedPtr<IPropertyHandle> ServiceProviderPropertyHandle;
+		IDetailCategoryBuilder* ServiceProviderCategoryBuilder;
+		TArray< TSharedPtr<FLocalizationServiceProviderWrapper> > Providers;
+
+		TSharedPtr<IPropertyHandle> TargetObjectsPropertyHandle;
+		FSimpleDelegate TargetsArrayPropertyHandle_OnNumElementsChanged;
+		TArray< TSharedPtr<IPropertyHandle> > TargetsList;
+		TSharedPtr< SListView< TSharedPtr<IPropertyHandle> > > TargetsListView;
+
+		/* If set, the entry at the index specified needs to be initialized as soon as possible. */
+		int32 NewEntryIndexToBeInitialized;
+	};
+	struct FLocalizationServiceProviderWrapper
+	{
+		FLocalizationServiceProviderWrapper() : Provider(nullptr) {}
+		FLocalizationServiceProviderWrapper(ILocalizationServiceProvider* const InProvider) : Provider(InProvider) {}
+
+		ILocalizationServiceProvider* Provider;
+	};
+
+	FProjectLocalizationSettingsDetailsCustomization::FProjectLocalizationSettingsDetailsCustomization(UProjectLocalizationSettings* const InSettings)
+		: Settings(InSettings)
+		, DetailLayoutBuilder(nullptr)
+		, ServiceProviderCategoryBuilder(nullptr)
+		, NewEntryIndexToBeInitialized(INDEX_NONE)
+		, IDetailCustomization()
+	{
+		TArray<ILocalizationServiceProvider*> ActualProviders = ILocalizationDashboardModule::Get().GetLocalizationServiceProviders();
+		for (ILocalizationServiceProvider* ActualProvider : ActualProviders)
+		{
+			TSharedPtr<FLocalizationServiceProviderWrapper> Provider = MakeShareable(new FLocalizationServiceProviderWrapper(ActualProvider));
+			Providers.Add(Provider);
+		}
+	}
+
+	class FindProviderPredicate
+	{
+	public:
+		FindProviderPredicate(ILocalizationServiceProvider* const InActualProvider)
+			: ActualProvider(InActualProvider)
+		{
+		}
+
+		bool operator()(const TSharedPtr<FLocalizationServiceProviderWrapper>& Provider)
+		{
+			return Provider->Provider == ActualProvider;
+		}
+
+	private:
+		ILocalizationServiceProvider* ActualProvider;
+	};
+
+
+	class FLocalizationDashboardCommands : public TCommands<FLocalizationDashboardCommands>
+	{
+	public:
+		FLocalizationDashboardCommands() 
+			: TCommands<FLocalizationDashboardCommands>("LocalizationDashboard", NSLOCTEXT("Contexts", "LocalizationDashboard", "Localization Dashboard"), NAME_None, FEditorStyle::GetStyleSetName())
+		{
+		}
+
+		TSharedPtr<FUICommandInfo> GatherAllTargets;
+		TSharedPtr<FUICommandInfo> ImportAllTargets;
+		TSharedPtr<FUICommandInfo> ExportAllTargets;
+
+		/** Initialize commands */
+		virtual void RegisterCommands() override;
+	};
+
+	void FLocalizationDashboardCommands::RegisterCommands() 
+	{
+		UI_COMMAND( GatherAllTargets, "Gather All", "Gathers text for all targets in the project.", EUserInterfaceActionType::Button, FInputGesture() );
+		UI_COMMAND( ImportAllTargets, "Import All", "Imports translations for all cultures of all targets in the project.", EUserInterfaceActionType::Button, FInputGesture() );
+		UI_COMMAND( ExportAllTargets, "Export All", "Exports translations for all cultures of all targets in the project.", EUserInterfaceActionType::Button, FInputGesture() );
+	}
+
+	void FProjectLocalizationSettingsDetailsCustomization::CustomizeDetails(IDetailLayoutBuilder& DetailBuilder)
+	{
+		DetailLayoutBuilder = &DetailBuilder;
+
+		{
+			ServiceProviderPropertyHandle = DetailBuilder.GetProperty(GET_MEMBER_NAME_CHECKED(UProjectLocalizationSettings, LocalizationServiceProvider));
+
+			if (ServiceProviderPropertyHandle.IsValid() && ServiceProviderPropertyHandle->IsValidHandle())
+			{
+				ServiceProviderPropertyHandle->MarkHiddenByCustomization();
+				//ServiceProviderCategoryBuilder = &( DetailBuilder.EditCategory( FName("LocalizationServiceProvider") ) );
+
+				//FDetailWidgetRow& ServiceProviderDetailWidgetRow = ServiceProviderCategoryBuilder->AddCustomRow( NSLOCTEXT("LocalizationDashboard", "LocalizationServiceProviderFilterString", "Localization Service Provider") );
+				//ServiceProviderDetailWidgetRow.NameContent()
+				//	[
+				//		ServiceProviderPropertyHandle->CreatePropertyNameWidget()
+				//	];
+
+				//ILocalizationServiceProvider* const LSP = ILocalizationDashboardModule::Get().GetCurrentLocalizationServiceProvider();
+				//TSharedPtr<FLocalizationServiceProviderWrapper>* LSPWrapper = Providers.FindByPredicate(FindProviderPredicate(LSP));
+				//ServiceProviderDetailWidgetRow.ValueContent()
+				//	[
+				//		SNew(SComboBox<TSharedPtr<FLocalizationServiceProviderWrapper>>)
+				//		.InitiallySelectedItem(LSPWrapper ? *LSPWrapper : nullptr)
+				//		.OptionsSource(&Providers)
+				//		.OnGenerateWidget(this, &FProjectLocalizationSettingsDetailsCustomization::ServiceProviderComboBox_OnGenerateWidget)
+				//		.OnSelectionChanged(this, &FProjectLocalizationSettingsDetailsCustomization::ServiceProviderComboBox_OnSelectionChanged)
+				//		.Content()
+				//		[
+				//			SNew(STextBlock)
+				//			.Text(this, &FProjectLocalizationSettingsDetailsCustomization::GetCurrentServiceProviderDisplayName)
+				//		]
+				//	];
+
+				//if (LSP)
+				//{
+				//	LSP->CustomizeSettingsDetails(*ServiceProviderCategoryBuilder);
+				//}
+			}
+		}
+
+		{
+			IDetailCategoryBuilder& DetailCategoryBuilder = DetailBuilder.EditCategory(FName("Targets"));
+			TargetObjectsPropertyHandle = DetailBuilder.GetProperty(GET_MEMBER_NAME_CHECKED(UProjectLocalizationSettings,TargetObjects));
+			if (TargetObjectsPropertyHandle->IsValidHandle())
+			{
+				TargetObjectsPropertyHandle->MarkHiddenByCustomization();
+				TargetsArrayPropertyHandle_OnNumElementsChanged = FSimpleDelegate::CreateSP(this, &FProjectLocalizationSettingsDetailsCustomization::RebuildTargetsList);
+				TargetObjectsPropertyHandle->AsArray()->SetOnNumElementsChanged(TargetsArrayPropertyHandle_OnNumElementsChanged);
+
+				FLocalizationDashboardCommands::Register();
+				const TSharedRef< FUICommandList > CommandList = MakeShareable(new FUICommandList);
+				FToolBarBuilder ToolBarBuilder( CommandList, FMultiBoxCustomization::AllowCustomization("LocalizationDashboard") );
+
+				CommandList->MapAction( FLocalizationDashboardCommands::Get().GatherAllTargets, FExecuteAction::CreateSP(this, &FProjectLocalizationSettingsDetailsCustomization::GatherAllTargets));
+				ToolBarBuilder.AddToolBarButton(FLocalizationDashboardCommands::Get().GatherAllTargets, NAME_None, TAttribute<FText>(), TAttribute<FText>(), FSlateIcon(FEditorStyle::GetStyleSetName(), "LocalizationDashboard.GatherAllTargets"));
+
+				CommandList->MapAction( FLocalizationDashboardCommands::Get().ImportAllTargets, FExecuteAction::CreateSP(this, &FProjectLocalizationSettingsDetailsCustomization::ImportAllTargets));
+				ToolBarBuilder.AddToolBarButton(FLocalizationDashboardCommands::Get().ImportAllTargets, NAME_None, TAttribute<FText>(), TAttribute<FText>(), FSlateIcon(FEditorStyle::GetStyleSetName(), "LocalizationDashboard.ImportForAllTargetsCultures"));
+
+				CommandList->MapAction( FLocalizationDashboardCommands::Get().ExportAllTargets, FExecuteAction::CreateSP(this, &FProjectLocalizationSettingsDetailsCustomization::ExportAllTargets));
+				ToolBarBuilder.AddToolBarButton(FLocalizationDashboardCommands::Get().ExportAllTargets, NAME_None, TAttribute<FText>(), TAttribute<FText>(), FSlateIcon(FEditorStyle::GetStyleSetName(), "LocalizationDashboard.ExportForAllTargetsCultures"));
+
+				BuildTargetsList();
+
+				DetailCategoryBuilder.AddCustomRow( NSLOCTEXT("LocalizationDashboard", "TargetsFilterString", "Targets") )
+					.WholeRowContent()
+					[
+						SNew(SVerticalBox)
+						+SVerticalBox::Slot()
+						.AutoHeight()
+						[
+							ToolBarBuilder.MakeWidget()
+						]
+						+SVerticalBox::Slot()
+							.AutoHeight()
+							[
+								SAssignNew(TargetsListView, SListView< TSharedPtr<IPropertyHandle> >)
+								.OnGenerateRow(this, &FProjectLocalizationSettingsDetailsCustomization::OnGenerateRow)
+								.ListItemsSource(&TargetsList)
+								.SelectionMode(ESelectionMode::None)
+								.HeaderRow
+								(
+								SNew(SHeaderRow)
+								+SHeaderRow::Column("Target")
+								.DefaultLabel( NSLOCTEXT("LocalizationDashboard", "TargetColumnLabel", "Target"))
+								.HAlignHeader(HAlign_Left)
+								.HAlignCell(HAlign_Left)
+								.VAlignCell(VAlign_Center)
+								+SHeaderRow::Column("Status")
+								.DefaultLabel( NSLOCTEXT("LocalizationDashboard", "StatusColumnLabel", "Status"))
+								.HAlignHeader(HAlign_Center)
+								.HAlignCell(HAlign_Center)
+								.VAlignCell(VAlign_Center)
+								+SHeaderRow::Column("Cultures")
+								.DefaultLabel( NSLOCTEXT("LocalizationDashboard", "CulturesColumnLabel", "Cultures"))
+								.HAlignHeader(HAlign_Fill)
+								.HAlignCell(HAlign_Fill)
+								.VAlignCell(VAlign_Center)
+								+SHeaderRow::Column("WordCount")
+								.DefaultLabel( NSLOCTEXT("LocalizationDashboard", "WordCountColumnLabel", "Word Count"))
+								.HAlignHeader(HAlign_Center)
+								.HAlignCell(HAlign_Center)
+								.VAlignCell(VAlign_Center)
+								+SHeaderRow::Column("Actions")
+								.DefaultLabel( NSLOCTEXT("LocalizationDashboard", "ActionsColumnLabel", "Actions"))
+								.HAlignHeader(HAlign_Center)
+								.HAlignCell(HAlign_Center)
+								.VAlignCell(VAlign_Center)
+								)
+							]
+						+SVerticalBox::Slot()
+							.AutoHeight()
+							.VAlign(VAlign_Center)
+							[
+								SNew(SButton)
+								.Text(NSLOCTEXT("LocalizationDashboard", "AddNewTargetButtonLabel", "Add New Target"))
+								.OnClicked(this, &FProjectLocalizationSettingsDetailsCustomization::OnNewTargetButtonClicked)
+							]
+					];
+			}
+		}
+	}
+
+	void FProjectLocalizationSettingsDetailsCustomization::BuildTargetsList()
+	{
+		const TSharedPtr<IPropertyHandleArray> TargetObjectsArrayPropertyHandle = TargetObjectsPropertyHandle->AsArray();
+		if (TargetObjectsArrayPropertyHandle.IsValid())
+		{
+			uint32 ElementCount = 0;
+			TargetObjectsArrayPropertyHandle->GetNumElements(ElementCount);
+			for (uint32 i = 0; i < ElementCount; ++i)
+			{
+				const TSharedPtr<IPropertyHandle> TargetObjectPropertyHandle = TargetObjectsArrayPropertyHandle->GetElement(i);
+				TargetsList.Add(TargetObjectPropertyHandle);
+			}
+		}
+	}
+
+	void FProjectLocalizationSettingsDetailsCustomization::RebuildTargetsList()
+	{
+		const TSharedPtr<IPropertyHandleArray> TargetObjectsArrayPropertyHandle = TargetObjectsPropertyHandle->AsArray();
+		if (TargetObjectsArrayPropertyHandle.IsValid() && NewEntryIndexToBeInitialized != INDEX_NONE)
+		{
+			const TSharedPtr<IPropertyHandle> TargetObjectPropertyHandle = TargetObjectsArrayPropertyHandle->GetElement(NewEntryIndexToBeInitialized);
+			if (TargetObjectPropertyHandle.IsValid() && TargetObjectPropertyHandle->IsValidHandle())
+			{
+				ULocalizationTarget* const NewTarget = Cast<ULocalizationTarget>(StaticConstructObject(ULocalizationTarget::StaticClass(), Settings));
+
+				TArray<void*> RawData;
+				TargetObjectsPropertyHandle->AccessRawData(RawData);
+				void* RawDatum = RawData.Top();
+				TArray<ULocalizationTarget*>& TargetObjectsArray = *(reinterpret_cast< TArray<ULocalizationTarget*>* >(RawDatum));
+
+				FString NewTargetName = "NewTarget";
+				auto TargetNameIsIdentical = [&](ULocalizationTarget* Target) -> bool
+				{
+					return (Target != NewTarget) && Target && (Target->Settings.Name == NewTargetName);
+				};
+
+				for (uint32 i = 0; TargetObjectsArray.ContainsByPredicate(TargetNameIsIdentical); ++i)
+				{
+					NewTargetName = FString::Printf(TEXT("NewTarget%u"), i);
+				}
+
+				NewTarget->Settings.Name = NewTargetName;
+				NewTarget->Settings.NativeCultureStatistics.CultureName = FInternationalization::Get().GetCurrentCulture()->GetName();
+
+				const UObject* SetValue = NewTarget;
+				TargetObjectPropertyHandle->SetValue(SetValue);
+
+				NewEntryIndexToBeInitialized = INDEX_NONE;
+			}
+		}
+
+		TargetsList.Empty();
+		BuildTargetsList();
+
+		if (TargetsListView.IsValid())
+		{
+			TargetsListView->RequestListRefresh();
+		}
+	}
+
+	FText FProjectLocalizationSettingsDetailsCustomization::GetCurrentServiceProviderDisplayName() const
+	{
+		ILocalizationServiceProvider* const LSP = ILocalizationDashboardModule::Get().GetCurrentLocalizationServiceProvider();
+		return LSP ? LSP->GetDisplayName() : NSLOCTEXT("LocalizationDashboard", "NoServiceProviderName", "None");
+	}
+
+	TSharedRef<SWidget> FProjectLocalizationSettingsDetailsCustomization::ServiceProviderComboBox_OnGenerateWidget(TSharedPtr<FLocalizationServiceProviderWrapper> LSPWrapper) const
+	{
+		ILocalizationServiceProvider* const LSP = LSPWrapper->Provider;
+		return	SNew(STextBlock)
+			.Text(LSP ? LSP->GetDisplayName() : NSLOCTEXT("LocalizationDashboard", "NoServiceProviderName", "None"));
+	}
+
+	void FProjectLocalizationSettingsDetailsCustomization::ServiceProviderComboBox_OnSelectionChanged(TSharedPtr<FLocalizationServiceProviderWrapper> LSPWrapper, ESelectInfo::Type SelectInfo)
+	{
+		ILocalizationServiceProvider* const LSP = LSPWrapper->Provider;
+		FString ServiceProviderName = LSP ? LSP->GetName().ToString() : TEXT("None");
+		ServiceProviderPropertyHandle->SetValue(ServiceProviderName);
+
+		if (LSP)
+		{
+			LSP->CustomizeSettingsDetails(*ServiceProviderCategoryBuilder);
+		}
+		DetailLayoutBuilder->ForceRefreshDetails();
+	}
+
+	void FProjectLocalizationSettingsDetailsCustomization::GatherAllTargets()
+	{
+		for (ULocalizationTarget* const LocalizationTarget : Settings->TargetObjects)
+		{
+			// Save unsaved packages.
+			bool DidPackagesNeedSaving;
+			const bool WerePackagesSaved = FEditorFileUtils::SaveDirtyPackages(true, true, true, false, false, &DidPackagesNeedSaving);
+
+			if (DidPackagesNeedSaving && !WerePackagesSaved)
+			{
+				// Give warning dialog.
+				const FText MessageText = NSLOCTEXT("LocalizationCultureActions", "UnsavedPackagesWarningDialogMessage", "There are unsaved changes. These changes may not be gathered from correctly.");
+				const FText TitleText = NSLOCTEXT("LocalizationCultureActions", "UnsavedPackagesWarningDialogTitle", "Unsaved Changes Before Gather");
+				switch(FMessageDialog::Open(EAppMsgType::OkCancel, MessageText, &TitleText))
+				{
+				case EAppReturnType::Cancel:
+					{
+						return;
+					}
+					break;
+				}
+			}
+
+		}
+
+		// Execute gather.
+		const TSharedPtr<SWindow> ParentWindow = FSlateApplication::Get().FindWidgetWindow(DetailLayoutBuilder->GetDetailsView().AsShared());
+		TArray<FLocalizationTargetSettings*> TargetSettings;
+		for (ULocalizationTarget* const TargetObject : Settings->TargetObjects)
+		{
+			TargetSettings.Add(&TargetObject->Settings);
+		}
+		LocalizationCommandletTasks::GatherTargets(ParentWindow.ToSharedRef(), TargetSettings);
+
+		for (ULocalizationTarget* const LocalizationTarget : Settings->TargetObjects)
+		{
+			UpdateTargetFromReports(LocalizationTarget);
+		}
+	}
+
+	void FProjectLocalizationSettingsDetailsCustomization::ImportAllTargets()
+	{
+		const TSharedPtr<SWindow> ParentWindow = FSlateApplication::Get().FindWidgetWindow(DetailLayoutBuilder->GetDetailsView().AsShared());
+		TArray<FLocalizationTargetSettings*> TargetSettings;
+		for (ULocalizationTarget* const TargetObject : Settings->TargetObjects)
+		{
+			TargetSettings.Add(&TargetObject->Settings);
+		}
+		LocalizationCommandletTasks::ImportTargets(ParentWindow.ToSharedRef(), TargetSettings);
+
+		for (ULocalizationTarget* const LocalizationTarget : Settings->TargetObjects)
+		{
+			UpdateTargetFromReports(LocalizationTarget);
+		}
+	}
+
+	void FProjectLocalizationSettingsDetailsCustomization::ExportAllTargets()
+	{
+		const TSharedPtr<SWindow> ParentWindow = FSlateApplication::Get().FindWidgetWindow(DetailLayoutBuilder->GetDetailsView().AsShared());
+		TArray<FLocalizationTargetSettings*> TargetSettings;
+		for (ULocalizationTarget* const TargetObject : Settings->TargetObjects)
+		{
+			TargetSettings.Add(&TargetObject->Settings);
+		}
+		LocalizationCommandletTasks::ExportTargets(ParentWindow.ToSharedRef(), TargetSettings);
+	}
+
+	void FProjectLocalizationSettingsDetailsCustomization::UpdateTargetFromReports(ULocalizationTarget* const LocalizationTarget)
+	{
+		//TArray< TSharedPtr<IPropertyHandle> > WordCountPropertyHandles;
+
+		//const TSharedPtr<IPropertyHandle> TargetSettingsPropertyHandle = TargetEditor->GetTargetSettingsPropertyHandle();
+		//if (TargetSettingsPropertyHandle.IsValid() && TargetSettingsPropertyHandle->IsValidHandle())
+		//{
+		//	const TSharedPtr<IPropertyHandle> NativeCultureStatisticsPropertyHandle = TargetSettingsPropertyHandle->GetChildHandle(GET_MEMBER_NAME_CHECKED(FLocalizationTargetSettings, NativeCultureStatistics));
+		//	if (NativeCultureStatisticsPropertyHandle.IsValid() && NativeCultureStatisticsPropertyHandle->IsValidHandle())
+		//	{
+		//		const TSharedPtr<IPropertyHandle> WordCountPropertyHandle = NativeCultureStatisticsPropertyHandle->GetChildHandle(GET_MEMBER_NAME_CHECKED(FCultureStatistics, WordCount));
+		//		if (WordCountPropertyHandle.IsValid() && WordCountPropertyHandle->IsValidHandle())
+		//		{
+		//			WordCountPropertyHandles.Add(WordCountPropertyHandle);
+		//		}
+		//	}
+
+		//	const TSharedPtr<IPropertyHandle> SupportedCulturesStatisticsPropertyHandle = TargetSettingsPropertyHandle->GetChildHandle(GET_MEMBER_NAME_CHECKED(FLocalizationTargetSettings, SupportedCulturesStatistics));
+		//	if (SupportedCulturesStatisticsPropertyHandle.IsValid() && SupportedCulturesStatisticsPropertyHandle->IsValidHandle())
+		//	{
+		//		uint32 SupportedCultureCount = 0;
+		//		SupportedCulturesStatisticsPropertyHandle->GetNumChildren(SupportedCultureCount);
+		//		for (uint32 i = 0; i < SupportedCultureCount; ++i)
+		//		{
+		//			const TSharedPtr<IPropertyHandle> ElementPropertyHandle = SupportedCulturesStatisticsPropertyHandle->GetChildHandle(i);
+		//			if (ElementPropertyHandle.IsValid() && ElementPropertyHandle->IsValidHandle())
+		//			{
+		//				const TSharedPtr<IPropertyHandle> WordCountPropertyHandle = SupportedCulturesStatisticsPropertyHandle->GetChildHandle(GET_MEMBER_NAME_CHECKED(FCultureStatistics, WordCount));
+		//				if (WordCountPropertyHandle.IsValid() && WordCountPropertyHandle->IsValidHandle())
+		//				{
+		//					WordCountPropertyHandles.Add(WordCountPropertyHandle);
+		//				}
+		//			}
+		//		}
+		//	}
+		//}
+
+		//for (const TSharedPtr<IPropertyHandle>& WordCountPropertyHandle : WordCountPropertyHandles)
+		//{
+		//	WordCountPropertyHandle->NotifyPreChange();
+		//}
+		LocalizationTarget->Settings.UpdateWordCountsFromCSV();
+		LocalizationTarget->Settings.UpdateStatusFromConflictReport();
+		//for (const TSharedPtr<IPropertyHandle>& WordCountPropertyHandle : WordCountPropertyHandles)
+		//{
+		//	WordCountPropertyHandle->NotifyPostChange();
+		//}
+	}
+
+	TSharedRef<ITableRow> FProjectLocalizationSettingsDetailsCustomization::OnGenerateRow(TSharedPtr<IPropertyHandle> TargetObjectPropertyHandle, const TSharedRef<STableViewBase>& Table)
+	{
+		return SNew(SLocalizationDashboardTargetRow, Table, DetailLayoutBuilder->GetPropertyUtilities(), TargetObjectPropertyHandle.ToSharedRef());
+	}
+
+	FReply FProjectLocalizationSettingsDetailsCustomization::OnNewTargetButtonClicked()
+	{
+		if(TargetObjectsPropertyHandle.IsValid() && TargetObjectsPropertyHandle->IsValidHandle())
+		{
+			uint32 NewElementIndex;
+			TargetObjectsPropertyHandle->AsArray()->GetNumElements(NewElementIndex);
+			TargetObjectsPropertyHandle->AsArray()->AddItem();
+			NewEntryIndexToBeInitialized = NewElementIndex;
+		}
+
+		return FReply::Handled();
+	}
+}
+
+void SLocalizationDashboard::Construct(const FArguments& InArgs, UProjectLocalizationSettings* const InSettings)
+{
+	check(InSettings);
+
+	FPropertyEditorModule& PropertyModule = FModuleManager::LoadModuleChecked<FPropertyEditorModule>("PropertyEditor");
+	FDetailsViewArgs DetailsViewArgs(false, false, false, FDetailsViewArgs::ENameAreaSettings::HideNameArea, false, nullptr, false, NAME_None);
+	TSharedRef<IDetailsView> DetailsView = PropertyModule.CreateDetailView(DetailsViewArgs);
+	DetailsView->RegisterInstancedCustomPropertyLayout(UProjectLocalizationSettings::StaticClass(), FOnGetDetailCustomizationInstance::CreateLambda([&]()->TSharedRef<IDetailCustomization>{return MakeShareable(new FProjectLocalizationSettingsDetailsCustomization(InSettings));}));
+
+	ChildSlot
+		[
+			DetailsView
+		];
+
+	DetailsView->SetObject(InSettings, true);
+}
+#undef LOCTEXT_NAMESPACE
\ No newline at end of file
diff --git a/Engine/Source/Editor/LocalizationDashboard/Private/SLocalizationDashboard.h b/Engine/Source/Editor/LocalizationDashboard/Private/SLocalizationDashboard.h
new file mode 100644
--- /dev/null
+++ b/Engine/Source/Editor/LocalizationDashboard/Private/SLocalizationDashboard.h
@@ -0,0 +1,19 @@
+// Copyright 1998-2014 Epic Games, Inc. All Rights Reserved.
+
+#pragma once
+
+#include "SCompoundWidget.h"
+
+class UProjectLocalizationSettings;
+
+class SLocalizationDashboard : public SCompoundWidget
+{
+public:
+	SLATE_BEGIN_ARGS(SLocalizationDashboard) {}
+	SLATE_END_ARGS()
+
+	void Construct(const FArguments& InArgs, UProjectLocalizationSettings* const Settings);
+
+public:
+	static const FName TabName;
+};
\ No newline at end of file
diff --git a/Engine/Source/Editor/LocalizationDashboard/Private/SLocalizationDashboardTargetRow.cpp b/Engine/Source/Editor/LocalizationDashboard/Private/SLocalizationDashboardTargetRow.cpp
new file mode 100644
--- /dev/null
+++ b/Engine/Source/Editor/LocalizationDashboard/Private/SLocalizationDashboardTargetRow.cpp
@@ -0,0 +1,450 @@
+// Copyright 1998-2014 Epic Games, Inc. All Rights Reserved.
+
+#include "LocalizationDashboardPrivatePCH.h"
+#include "SLocalizationDashboard.h"
+#include "SLocalizationDashboardTargetRow.h"
+#include "SHyperlink.h"
+#include "SLocalizationTargetEditor.h"
+#include "SLocalizationTargetStatusButton.h"
+#include "LocalizationConfigurationScript.h"
+#include "LocalizationCommandletTasks.h"
+
+TMap< TWeakObjectPtr<ULocalizationTarget>, TWeakPtr<SDockTab> > SLocalizationDashboardTargetRow::TargetToTabMap;
+
+void SLocalizationDashboardTargetRow::Construct(const FTableRowArgs& InArgs, const TSharedRef<STableViewBase>& OwnerTableView, const TSharedRef<IPropertyUtilities>& InPropertyUtilities, const TSharedRef<IPropertyHandle>& InTargetObjectPropertyHandle)
+{
+	PropertyUtilities = InPropertyUtilities;
+	TargetObjectPropertyHandle = InTargetObjectPropertyHandle;
+
+	FSuperRowType::Construct(InArgs, OwnerTableView);
+}
+
+TSharedRef<SWidget> SLocalizationDashboardTargetRow::GenerateWidgetForColumn( const FName& ColumnName )
+{
+	TSharedPtr<SWidget> Return;
+
+	if (ColumnName == "Target")
+	{
+		// Target Name
+		Return = SNew(SHyperlink)
+			.Text(this, &SLocalizationDashboardTargetRow::GetTargetName)
+			.OnNavigate(this, &SLocalizationDashboardTargetRow::OnNavigate);
+	}
+	else if(ColumnName == "Status")
+	{
+		ULocalizationTarget* const LocalizationTarget = GetTarget();
+		if (LocalizationTarget)
+		{
+			// Status icon button.
+			Return = SNew(SLocalizationTargetStatusButton, *LocalizationTarget);
+		}
+	}
+	else if (ColumnName == "Cultures")
+	{
+		// Culture Names
+		Return = SNew(STextBlock)
+			.Text(this, &SLocalizationDashboardTargetRow::GetCulturesText);
+	}
+	else if(ColumnName == "WordCount")
+	{
+		// Progress Bar and Word Counts
+		Return = SNew(STextBlock)
+			.Text(this, &SLocalizationDashboardTargetRow::GetWordCountText);
+	}
+	else if(ColumnName == "Actions")
+	{
+		TSharedRef<SHorizontalBox> HorizontalBox = SNew(SHorizontalBox);
+		Return = HorizontalBox;
+
+		// Gather
+		HorizontalBox->AddSlot()
+			.FillWidth(1.0f)
+			.HAlign(HAlign_Center)
+			.VAlign(VAlign_Center)
+			[
+				SNew(SButton)
+				.ButtonStyle( FEditorStyle::Get(), TEXT("HoverHintOnly") )
+				.ToolTipText( NSLOCTEXT("LocalizationTargetActions", "GatherButtonLabel", "Gather") )
+				.OnClicked(this, &SLocalizationDashboardTargetRow::Gather)
+				.Content()
+				[
+					SNew(SImage)
+					.Image( FEditorStyle::GetBrush("LocalizationDashboard.GatherTarget") )
+				]
+			];
+
+		// Refresh Word Counts
+		HorizontalBox->AddSlot()
+			.FillWidth(1.0f)
+			.HAlign(HAlign_Center)
+			.VAlign(VAlign_Center)
+			[
+				SNew(SButton)
+				.ButtonStyle( FEditorStyle::Get(), TEXT("HoverHintOnly") )
+				.ToolTipText(NSLOCTEXT("LocalizationTargetActions", "RefreshWordCountButtonLabel", "Refresh"))
+				.OnClicked(this, &SLocalizationDashboardTargetRow::RefreshWordCount)
+				.Content()
+				[
+					SNew(SImage)
+					.Image(FEditorStyle::GetBrush("LocalizationDashboard.RefreshWordCount"))
+				]
+			];
+
+		// Import All
+		HorizontalBox->AddSlot()
+			.FillWidth(1.0f)
+			.HAlign(HAlign_Center)
+			.VAlign(VAlign_Center)
+			[
+				SNew(SButton)
+				.ButtonStyle( FEditorStyle::Get(), TEXT("HoverHintOnly") )
+				.ToolTipText( NSLOCTEXT("LocalizationTargetActions", "ImportAllButtonLabel", "Import All") )
+				.OnClicked(this, &SLocalizationDashboardTargetRow::ImportAll)
+				.Content()
+				[
+					SNew(SImage)
+					.Image( FEditorStyle::GetBrush("LocalizationDashboard.ImportForAllCultures") )
+				]
+			];
+
+		// Export All
+		HorizontalBox->AddSlot()
+			.FillWidth(1.0f)
+			.HAlign(HAlign_Center)
+			.VAlign(VAlign_Center)
+			[
+				SNew(SButton)
+				.ButtonStyle( FEditorStyle::Get(), TEXT("HoverHintOnly") )
+				.ToolTipText(NSLOCTEXT("LocalizationTargetActions", "ExportAllButtonLabel", "Export All"))
+				.OnClicked(this, &SLocalizationDashboardTargetRow::ExportAll)
+				.Content()
+				[
+					SNew(SImage)
+					.Image(FEditorStyle::GetBrush("LocalizationDashboard.ExportForAllCultures"))
+				]
+			];
+
+		// Delete Target
+		HorizontalBox->AddSlot()
+			.FillWidth(1.0f)
+			.HAlign(HAlign_Center)
+			.VAlign(VAlign_Center)
+			[
+				SNew(SButton)
+				.ButtonStyle( FEditorStyle::Get(), TEXT("HoverHintOnly") )
+				.ToolTipText(NSLOCTEXT("LocalizationTargetActions", "DeleteButtonLabel", "Delete"))
+				.OnClicked(this, &SLocalizationDashboardTargetRow::EnqueueDeletion)
+				.Content()
+				[
+					SNew(SImage)
+					.Image(FEditorStyle::GetBrush("LocalizationDashboard.DeleteTarget"))
+				]
+			];
+	}
+
+	return Return.IsValid() ? Return.ToSharedRef() : SNullWidget::NullWidget;
+}
+
+ULocalizationTarget* SLocalizationDashboardTargetRow::GetTarget() const
+{
+	if (TargetObjectPropertyHandle.IsValid() && TargetObjectPropertyHandle->IsValidHandle())
+	{
+		UObject* Object = nullptr;
+		if(FPropertyAccess::Result::Success == TargetObjectPropertyHandle->GetValue(Object))
+		{
+			return Cast<ULocalizationTarget>(Object);
+		}
+	}
+	return nullptr;
+}
+
+FLocalizationTargetSettings* SLocalizationDashboardTargetRow::GetTargetSettings() const
+{
+	ULocalizationTarget* const LocalizationTarget = GetTarget();
+	return LocalizationTarget ? &(LocalizationTarget->Settings) : nullptr;
+}
+
+FText SLocalizationDashboardTargetRow::GetTargetName() const
+{
+	ULocalizationTarget* const LocalizationTarget = GetTarget();
+	return LocalizationTarget ? FText::FromString(LocalizationTarget->Settings.Name) : FText::GetEmpty();
+}
+
+void SLocalizationDashboardTargetRow::OnNavigate()
+{
+	ULocalizationTarget* const LocalizationTarget = GetTarget();
+	if (LocalizationTarget)
+	{
+		// Create tab if not existent.
+		TWeakPtr<SDockTab>& TargetEditorDockTab = TargetToTabMap.FindOrAdd(TWeakObjectPtr<ULocalizationTarget>(LocalizationTarget));
+
+		if (!TargetEditorDockTab.IsValid())
+		{
+			UProjectLocalizationSettings* ProjectSettings = nullptr;
+			const TSharedPtr<IPropertyHandle> TargetObjectsPropertyHandle = TargetObjectPropertyHandle->GetParentHandle(); // TargetObjects
+			if (TargetObjectsPropertyHandle.IsValid() && TargetObjectsPropertyHandle->IsValidHandle())
+			{
+				TArray<UObject*> OuterObjects;
+				TargetObjectsPropertyHandle->GetOuterObjects(OuterObjects);
+				ProjectSettings = Cast<UProjectLocalizationSettings>(OuterObjects.Top());
+			}
+
+			const TSharedRef<SLocalizationTargetEditor> OurTargetEditor = SNew(SLocalizationTargetEditor, ProjectSettings, LocalizationTarget);
+			const TSharedRef<SDockTab> NewTargetEditorTab = SNew(SDockTab)
+				.TabRole(ETabRole::DocumentTab)
+				.Label_Lambda( [LocalizationTarget]
+			{
+				return LocalizationTarget ? FText::FromString(LocalizationTarget->Settings.Name) : FText::GetEmpty();
+			})
+				[
+					OurTargetEditor
+				];
+
+			FGlobalTabmanager::Get()->InsertNewDocumentTab( SLocalizationDashboard::TabName, FTabManager::ESearchPreference::RequireClosedTab, NewTargetEditorTab );
+			TargetEditorDockTab = NewTargetEditorTab;
+		}
+		else
+		{
+			const TSharedPtr<SDockTab> OldTargetEditorTab = TargetEditorDockTab.Pin();
+			FGlobalTabmanager::Get()->DrawAttention(OldTargetEditorTab.ToSharedRef());
+		}
+	}
+}
+
+FText SLocalizationDashboardTargetRow::GetCulturesText() const
+{
+	FLocalizationTargetSettings* const TargetSettings = GetTargetSettings();
+	if (TargetSettings)
+	{
+		TArray<FString> OrderedCultureNames;
+		OrderedCultureNames.Add(TargetSettings->NativeCultureStatistics.CultureName);
+		for (auto& ForeignCultureStatistics : TargetSettings->SupportedCulturesStatistics)
+		{
+			OrderedCultureNames.Add(ForeignCultureStatistics.CultureName);
+		}
+
+		FString Result;
+		for (FString& CultureName : OrderedCultureNames)
+		{
+			const FCulturePtr Culture = FInternationalization::Get().GetCulture(CultureName);
+			if (Culture.IsValid())
+			{
+				const FString CultureDisplayName = Culture->GetDisplayName();
+
+				if (!Result.IsEmpty())
+				{
+					Result.Append(TEXT(", "));
+				}
+
+				Result.Append(CultureDisplayName);
+			}
+		}
+
+		return FText::FromString(Result);
+	}
+
+	return FText::GetEmpty();
+}
+
+uint32 SLocalizationDashboardTargetRow::GetWordCount() const
+{
+	FLocalizationTargetSettings* const TargetSettings = GetTargetSettings();
+	return TargetSettings ? TargetSettings->NativeCultureStatistics.WordCount : 0;
+}
+
+uint32 SLocalizationDashboardTargetRow::GetNativeWordCount() const
+{
+	FLocalizationTargetSettings* const TargetSettings = GetTargetSettings();
+	return TargetSettings ? TargetSettings->NativeCultureStatistics.WordCount : 0;
+}
+
+FText SLocalizationDashboardTargetRow::GetWordCountText() const
+{
+	return FText::AsNumber(GetWordCount());
+}
+
+void SLocalizationDashboardTargetRow::UpdateTargetFromReports()
+{
+	ULocalizationTarget* const LocalizationTarget = GetTarget();
+	if (LocalizationTarget)
+	{
+		//TArray< TSharedPtr<IPropertyHandle> > WordCountPropertyHandles;
+
+		//const TSharedPtr<IPropertyHandle> TargetSettingsPropertyHandle = TargetEditor->GetTargetSettingsPropertyHandle();
+		//if (TargetSettingsPropertyHandle.IsValid() && TargetSettingsPropertyHandle->IsValidHandle())
+		//{
+		//	const TSharedPtr<IPropertyHandle> NativeCultureStatisticsPropertyHandle = TargetSettingsPropertyHandle->GetChildHandle(GET_MEMBER_NAME_CHECKED(FLocalizationTargetSettings, NativeCultureStatistics));
+		//	if (NativeCultureStatisticsPropertyHandle.IsValid() && NativeCultureStatisticsPropertyHandle->IsValidHandle())
+		//	{
+		//		const TSharedPtr<IPropertyHandle> WordCountPropertyHandle = NativeCultureStatisticsPropertyHandle->GetChildHandle(GET_MEMBER_NAME_CHECKED(FCultureStatistics, WordCount));
+		//		if (WordCountPropertyHandle.IsValid() && WordCountPropertyHandle->IsValidHandle())
+		//		{
+		//			WordCountPropertyHandles.Add(WordCountPropertyHandle);
+		//		}
+		//	}
+
+		//	const TSharedPtr<IPropertyHandle> SupportedCulturesStatisticsPropertyHandle = TargetSettingsPropertyHandle->GetChildHandle(GET_MEMBER_NAME_CHECKED(FLocalizationTargetSettings, SupportedCulturesStatistics));
+		//	if (SupportedCulturesStatisticsPropertyHandle.IsValid() && SupportedCulturesStatisticsPropertyHandle->IsValidHandle())
+		//	{
+		//		uint32 SupportedCultureCount = 0;
+		//		SupportedCulturesStatisticsPropertyHandle->GetNumChildren(SupportedCultureCount);
+		//		for (uint32 i = 0; i < SupportedCultureCount; ++i)
+		//		{
+		//			const TSharedPtr<IPropertyHandle> ElementPropertyHandle = SupportedCulturesStatisticsPropertyHandle->GetChildHandle(i);
+		//			if (ElementPropertyHandle.IsValid() && ElementPropertyHandle->IsValidHandle())
+		//			{
+		//				const TSharedPtr<IPropertyHandle> WordCountPropertyHandle = SupportedCulturesStatisticsPropertyHandle->GetChildHandle(GET_MEMBER_NAME_CHECKED(FCultureStatistics, WordCount));
+		//				if (WordCountPropertyHandle.IsValid() && WordCountPropertyHandle->IsValidHandle())
+		//				{
+		//					WordCountPropertyHandles.Add(WordCountPropertyHandle);
+		//				}
+		//			}
+		//		}
+		//	}
+		//}
+
+		//for (const TSharedPtr<IPropertyHandle>& WordCountPropertyHandle : WordCountPropertyHandles)
+		//{
+		//	WordCountPropertyHandle->NotifyPreChange();
+		//}
+		LocalizationTarget->Settings.UpdateWordCountsFromCSV();
+		LocalizationTarget->Settings.UpdateStatusFromConflictReport();
+		//for (const TSharedPtr<IPropertyHandle>& WordCountPropertyHandle : WordCountPropertyHandles)
+		//{
+		//	WordCountPropertyHandle->NotifyPostChange();
+		//}
+	}
+}
+
+FReply SLocalizationDashboardTargetRow::RefreshWordCount()
+{
+	bool HasSucceeded = false;
+
+	ULocalizationTarget* const LocalizationTarget = GetTarget();
+	if (LocalizationTarget)
+	{
+		const TSharedPtr<SWindow> ParentWindow = FSlateApplication::Get().FindWidgetWindow(AsShared());
+		LocalizationCommandletTasks::GenerateReportsForTarget(ParentWindow.ToSharedRef(), LocalizationTarget->Settings);
+
+		UpdateTargetFromReports();
+	}
+
+	return FReply::Handled();
+}
+
+FReply SLocalizationDashboardTargetRow::Gather()
+{
+	ULocalizationTarget* const LocalizationTarget = GetTarget();
+	if (LocalizationTarget)
+	{
+		// Save unsaved packages.
+		bool DidPackagesNeedSaving;
+		const bool WerePackagesSaved = FEditorFileUtils::SaveDirtyPackages(true, true, true, false, false, &DidPackagesNeedSaving);
+
+		if (DidPackagesNeedSaving && !WerePackagesSaved)
+		{
+			// Give warning dialog.
+			const FText MessageText = NSLOCTEXT("LocalizationCultureActions", "UnsavedPackagesWarningDialogMessage", "There are unsaved changes. These changes may not be gathered from correctly.");
+			const FText TitleText = NSLOCTEXT("LocalizationCultureActions", "UnsavedPackagesWarningDialogTitle", "Unsaved Changes Before Gather");
+			switch(FMessageDialog::Open(EAppMsgType::OkCancel, MessageText, &TitleText))
+			{
+			case EAppReturnType::Cancel:
+				{
+					return FReply::Handled();
+				}
+				break;
+			}
+		}
+
+		// Execute gather.
+		const TSharedPtr<SWindow> ParentWindow = FSlateApplication::Get().FindWidgetWindow(AsShared());
+		LocalizationCommandletTasks::GatherTarget(ParentWindow.ToSharedRef(), LocalizationTarget->Settings);
+
+		UpdateTargetFromReports();
+	}
+
+	return FReply::Handled();
+}
+
+FReply SLocalizationDashboardTargetRow::ImportAll()
+{
+	ULocalizationTarget* const LocalizationTarget = GetTarget();
+	if (LocalizationTarget)
+	{
+		const TSharedPtr<SWindow> ParentWindow = FSlateApplication::Get().FindWidgetWindow(AsShared());
+		LocalizationCommandletTasks::ImportTarget(ParentWindow.ToSharedRef(), LocalizationTarget->Settings);
+
+		UpdateTargetFromReports();
+	}
+
+	return FReply::Handled();
+}
+
+FReply SLocalizationDashboardTargetRow::ExportAll()
+{
+	ULocalizationTarget* const LocalizationTarget = GetTarget();
+	if (LocalizationTarget)
+	{
+		const TSharedPtr<SWindow> ParentWindow = FSlateApplication::Get().FindWidgetWindow(AsShared());
+		LocalizationCommandletTasks::ExportTarget(ParentWindow.ToSharedRef(), LocalizationTarget->Settings);
+	}
+
+	return FReply::Handled();
+}
+
+FReply SLocalizationDashboardTargetRow::EnqueueDeletion()
+{
+	PropertyUtilities->EnqueueDeferredAction(FSimpleDelegate::CreateSP(this, &SLocalizationDashboardTargetRow::Delete));
+	return FReply::Handled();
+}
+
+void SLocalizationDashboardTargetRow::Delete()
+{
+	static bool IsExecuting = false;
+	if (!IsExecuting)
+	{
+		TGuardValue<bool> ReentranceGuard(IsExecuting, true);
+
+		ULocalizationTarget* const LocalizationTarget = GetTarget();
+		if (LocalizationTarget)
+		{
+			FText TitleText;
+			FText MessageText;
+
+			// Setup dialog for deleting target.
+			const FText FormatPattern = NSLOCTEXT("LocalizationDashboard", "DeleteTargetConfirmationDialogMessage", "This action can not be undone and data will be permanently lost. Are you sure you would like to delete {TargetName}?");
+			FFormatNamedArguments Arguments;
+			Arguments.Add(TEXT("TargetName"), FText::FromString(LocalizationTarget->Settings.Name));
+			MessageText = FText::Format(FormatPattern, Arguments);
+			TitleText = NSLOCTEXT("LocalizationDashboard", "DeleteTargetConfirmationDialogTitle", "Confirm Target Deletion");
+
+			switch(FMessageDialog::Open(EAppMsgType::OkCancel, MessageText, &TitleText))
+			{
+			case EAppReturnType::Ok:
+				{
+					LocalizationTarget->Settings.DeleteFiles();
+
+					// Close target editor.
+					TWeakPtr<SDockTab>& TargetEditorDockTab = TargetToTabMap.FindOrAdd(TWeakObjectPtr<ULocalizationTarget>(LocalizationTarget));
+					if (TargetEditorDockTab.IsValid())
+					{
+						const TSharedPtr<SDockTab> OldTargetEditorTab = TargetEditorDockTab.Pin();
+						OldTargetEditorTab->RequestCloseTab();
+					}
+
+					// Remove this element from the parent array.
+					const TSharedPtr<IPropertyHandle> ParentPropertyHandle = TargetObjectPropertyHandle->GetParentHandle();
+					if (ParentPropertyHandle.IsValid() && ParentPropertyHandle->IsValidHandle())
+					{
+						const TSharedPtr<IPropertyHandleArray> ParentArrayPropertyHandle = ParentPropertyHandle->AsArray();
+						if (ParentArrayPropertyHandle.IsValid())
+						{
+							ParentArrayPropertyHandle->DeleteItem(TargetObjectPropertyHandle->GetIndexInArray());
+						}
+					}
+				}
+				break;
+			}
+		}
+	}
+}
diff --git a/Engine/Source/Editor/LocalizationDashboard/Private/SLocalizationDashboardTargetRow.h b/Engine/Source/Editor/LocalizationDashboard/Private/SLocalizationDashboardTargetRow.h
new file mode 100644
--- /dev/null
+++ b/Engine/Source/Editor/LocalizationDashboard/Private/SLocalizationDashboardTargetRow.h
@@ -0,0 +1,44 @@
+// Copyright 1998-2014 Epic Games, Inc. All Rights Reserved.
+#pragma once
+
+#include "STableRow.h"
+#include "IPropertyUtilities.h"
+#include "PropertyHandle.h"
+#include "ProjectLocalizationSettings.h"
+#include "LocalizationTarget.h"
+#include "SDockTab.h"
+
+class SLocalizationDashboardTargetRow : public SMultiColumnTableRow< TSharedPtr<IPropertyHandle> >
+{
+public:
+	void Construct(const FTableRowArgs& InArgs, const TSharedRef<STableViewBase>& OwnerTableView, const TSharedRef<IPropertyUtilities>& InPropertyUtilities, const TSharedRef<IPropertyHandle>& InTargetObjectPropertyHandle);
+	TSharedRef<SWidget> GenerateWidgetForColumn( const FName& ColumnName ) override;
+
+private:
+	ULocalizationTarget* GetTarget() const;
+	FLocalizationTargetSettings* GetTargetSettings() const;
+
+	FText GetTargetName() const;
+	void OnNavigate();
+
+	FText GetCulturesText() const;
+
+	uint32 GetWordCount() const;
+	uint32 GetNativeWordCount() const;
+	FText GetWordCountText() const;
+
+	void UpdateTargetFromReports();
+
+	FReply RefreshWordCount();
+
+	FReply Gather();
+	FReply ImportAll();
+	FReply ExportAll();
+	FReply EnqueueDeletion();
+	void Delete();
+
+private:
+	TSharedPtr<IPropertyUtilities> PropertyUtilities;
+	TSharedPtr<IPropertyHandle> TargetObjectPropertyHandle;
+	static TMap< TWeakObjectPtr<ULocalizationTarget>, TWeakPtr<SDockTab> > TargetToTabMap;
+};
\ No newline at end of file
diff --git a/Engine/Source/Editor/LocalizationDashboard/Private/SLocalizationTargetEditor.cpp b/Engine/Source/Editor/LocalizationDashboard/Private/SLocalizationTargetEditor.cpp
new file mode 100644
--- /dev/null
+++ b/Engine/Source/Editor/LocalizationDashboard/Private/SLocalizationTargetEditor.cpp
@@ -0,0 +1,879 @@
+// Copyright 1998-2014 Epic Games, Inc. All Rights Reserved.
+
+#include "LocalizationDashboardPrivatePCH.h"
+#include "IDetailCustomization.h"
+#include "EditorStyleSet.h"
+#include "DetailLayoutBuilder.h"
+#include "LocalizationConfigurationScript.h"
+#include "SCulturePicker.h"
+#include "SLocalizationTargetEditor.h"
+#include "PropertyEditorModule.h"
+#include "IDetailsView.h"
+#include "DetailWidgetRow.h"
+#include "DetailCategoryBuilder.h"
+#include "PropertyHandle.h"
+#include "SLocalizationTargetEditorCultureRow.h"
+#include "TargetsTableEntry.h"
+#include "SCulturePicker.h"
+#include "IDetailPropertyRow.h"
+#include "SLocalizationTargetStatusButton.h"
+#include "LocalizationCommandletTasks.h"
+
+#define LOCTEXT_NAMSPACE "LocalizationCulture"
+
+namespace
+{
+	class FLocalizationTargetDetailCustomization : public IDetailCustomization
+	{
+	public:
+		FLocalizationTargetDetailCustomization(UProjectLocalizationSettings* const InProjectSettings, ULocalizationTarget* const InLocalizationTarget);
+		virtual void CustomizeDetails(IDetailLayoutBuilder& DetailBuilder) override;
+
+		FLocalizationTargetSettings* GetTargetSettings() const;
+		TSharedPtr<IPropertyHandle> GetTargetSettingsPropertyHandle() const;
+
+	private:
+		FText GetTargetName() const;
+		bool IsTargetNameUnique(const FString& Name) const;
+		void OnTargetNameChanged(const FText& NewText);
+		void OnTargetNameCommitted(const FText& NewText, ETextCommit::Type Type);
+
+		void RebuildTargetDependenciesBox();
+		void RebuildTargetsList();
+		TSharedRef<ITableRow> OnGenerateTargetRow(ULocalizationTarget* OtherLocalizationTarget, const TSharedRef<STableViewBase>& Table);
+		void OnTargetDependencyCheckStateChanged(ULocalizationTarget* const OtherLocalizationTarget, const ECheckBoxState State);
+		ECheckBoxState IsTargetDependencyChecked(ULocalizationTarget* const OtherLocalizationTarget) const;
+
+		FText GetNativeCultureName() const;
+		FText GetNativeCultureDisplayName() const;
+		void OnNativeCultureSelected(FCulturePtr SelectedCulture, ESelectInfo::Type SelectInfo);
+
+		void Gather();
+		void ImportAllCultures();
+		void ExportAllCultures();
+		void RefreshWordCounts();
+		void UpdateTargetFromReports();
+
+		void BuildListedCulturesList();
+		void RebuildListedCulturesList();
+		TSharedRef<ITableRow> OnGenerateCultureRow(TSharedPtr<IPropertyHandle> CulturePropertyHandle, const TSharedRef<STableViewBase>& Table);
+
+		bool IsCultureSelectableAsSupported(FCulturePtr Culture);
+		void OnNewSupportedCultureSelected(FCulturePtr SelectedCulture, ESelectInfo::Type SelectInfo);
+
+	private:
+		UProjectLocalizationSettings* ProjectSettings;
+		ULocalizationTarget* LocalizationTarget;
+
+		IDetailLayoutBuilder* DetailLayoutBuilder;
+		TSharedPtr<IPropertyHandle> TargetSettingsPropertyHandle;
+
+		TSharedPtr<SEditableTextBox> TargetNameEditableTextBox;
+
+		TSharedPtr<SHorizontalBox> TargetDependenciesHorizontalBox;
+		TArray< TSharedPtr<SWidget> > TargetDependenciesWidgets;
+		TArray<ULocalizationTarget*> TargetDependenciesOptionsList;
+		TSharedPtr< SListView<ULocalizationTarget*> > TargetDependenciesListView;
+
+		TSharedPtr<IPropertyHandle> NativeCultureStatisticsPropertyHandle;
+		TSharedPtr<IPropertyHandle> NativeCultureNamePropertyHandle;
+		TSharedPtr<SComboButton> NativeCultureComboButton;
+		TArray<FCulturePtr> AllCultures;
+
+		TSharedPtr<IPropertyHandle> SupportedCulturesStatisticsPropertyHandle;
+		FSimpleDelegate SupportedCulturesStatisticsPropertyHandle_OnNumElementsChanged;
+		TSharedPtr< SListView< TSharedPtr<IPropertyHandle> > > SupportedCultureListView;
+		TSharedPtr<SComboButton> AddNewSupportedCultureComboButton;
+		TSharedPtr<SCulturePicker> SupportedCulturePicker;
+		TArray< TSharedPtr<IPropertyHandle> > ListedCultureStatisticProperties;
+
+		/* If set, the entry at the index specified needs to be initialized as soon as possible. */
+		int32 NewEntryIndexToBeInitialized;
+		FCulturePtr SelectedNewCulture;
+	};
+
+	FLocalizationTargetDetailCustomization::FLocalizationTargetDetailCustomization(UProjectLocalizationSettings* const InProjectSettings, ULocalizationTarget* const InLocalizationTarget)
+		: ProjectSettings(InProjectSettings)
+		, LocalizationTarget(InLocalizationTarget)
+		, DetailLayoutBuilder(nullptr)
+		, NewEntryIndexToBeInitialized(INDEX_NONE)
+	{
+	}
+
+	class FLocalizationTargetEditorCommands : public TCommands<FLocalizationTargetEditorCommands>
+	{
+	public:
+		FLocalizationTargetEditorCommands() 
+			: TCommands<FLocalizationTargetEditorCommands>("LocalizationTargetEditor", NSLOCTEXT("Contexts", "LocalizationTargetEditor", "Localization Target Editor"), NAME_None, FEditorStyle::GetStyleSetName())
+		{
+		}
+
+		TSharedPtr<FUICommandInfo> Gather;
+		TSharedPtr<FUICommandInfo> ImportAllCultures;
+		TSharedPtr<FUICommandInfo> ExportAllCultures;
+		TSharedPtr<FUICommandInfo> RefreshWordCounts;
+
+		/** Initialize commands */
+		virtual void RegisterCommands() override;
+	};
+
+	void FLocalizationTargetEditorCommands::RegisterCommands() 
+	{
+		UI_COMMAND( Gather, "Gather", "Gathers text for this target.", EUserInterfaceActionType::Button, FInputGesture() );
+		UI_COMMAND( ImportAllCultures, "Import All", "Imports translations for all cultures of this target.", EUserInterfaceActionType::Button, FInputGesture() );
+		UI_COMMAND( ExportAllCultures, "Export All", "Exports translations for all cultures of this target.", EUserInterfaceActionType::Button, FInputGesture() );
+		UI_COMMAND( RefreshWordCounts, "Count Words", "Recalculates the word counts for all cultures of this target.", EUserInterfaceActionType::Button, FInputGesture() );
+	}
+
+	void FLocalizationTargetDetailCustomization::CustomizeDetails(IDetailLayoutBuilder& DetailBuilder)
+	{
+		DetailLayoutBuilder = &DetailBuilder;
+		TargetSettingsPropertyHandle = DetailBuilder.GetProperty(GET_MEMBER_NAME_CHECKED(ULocalizationTarget, Settings));
+
+		if (TargetSettingsPropertyHandle.IsValid() && TargetSettingsPropertyHandle->IsValidHandle())
+		{
+			TargetSettingsPropertyHandle->MarkHiddenByCustomization();
+		}
+
+		//{
+		//	IDetailCategoryBuilder& ServiceProviderCategoryBuilder = DetailBuilder.EditCategory(FName("LocalizationServiceProvider"));
+
+		//	ILocalizationServiceProvider* const LSP = ILocalizationDashboardModule::Get().GetCurrentLocalizationServiceProvider();
+		//	if (LSP && LocalizationTarget)
+		//	{
+		//		LSP->CustomizeTargetDetails(ServiceProviderCategoryBuilder, *LocalizationTarget);
+		//	}
+		//}
+
+		if (TargetSettingsPropertyHandle.IsValid() && TargetSettingsPropertyHandle->IsValidHandle())
+		{
+			IDetailCategoryBuilder& DetailCategoryBuilder = DetailBuilder.EditCategory(FName("Target"));
+
+			const TSharedPtr<IPropertyHandle> NamePropertyHandle = TargetSettingsPropertyHandle->GetChildHandle(GET_MEMBER_NAME_CHECKED(FLocalizationTargetSettings, Name));
+			if (NamePropertyHandle.IsValid() && NamePropertyHandle->IsValidHandle())
+			{
+				FDetailWidgetRow& DetailWidgetRow = DetailCategoryBuilder.AddCustomRow( NSLOCTEXT("LocalizationTargetEditor", "TargetNameFilterString", "Target Name") );
+				DetailWidgetRow.NameContent()
+					[
+						NamePropertyHandle->CreatePropertyNameWidget()
+					];
+				DetailWidgetRow.ValueContent()
+					[
+						SAssignNew(TargetNameEditableTextBox, SEditableTextBox)
+						.Text(this, &FLocalizationTargetDetailCustomization::GetTargetName)
+						.RevertTextOnEscape(true)
+						.OnTextChanged(this, &FLocalizationTargetDetailCustomization::OnTargetNameChanged)
+						.OnTextCommitted(this, &FLocalizationTargetDetailCustomization::OnTargetNameCommitted)
+					];
+			}
+			const TSharedPtr<IPropertyHandle> StatusPropertyHandle = TargetSettingsPropertyHandle->GetChildHandle(GET_MEMBER_NAME_CHECKED(FLocalizationTargetSettings, Status));
+			if (StatusPropertyHandle.IsValid() && StatusPropertyHandle->IsValidHandle())
+			{
+				StatusPropertyHandle->MarkHiddenByCustomization();
+				if (LocalizationTarget)
+				{
+					FDetailWidgetRow& StatusRow = DetailCategoryBuilder.AddCustomRow( NSLOCTEXT("LocalizationTargetEditor", "ConflictReportStatusFilterString", "Conflict Report Status") );
+					StatusRow.NameContent()
+						[
+							StatusPropertyHandle->CreatePropertyNameWidget()
+						];
+					StatusRow.ValueContent()
+						.HAlign(HAlign_Left)
+						.VAlign(VAlign_Center)
+						[
+							SNew(SLocalizationTargetStatusButton, *LocalizationTarget)
+						];
+				}
+			}
+
+			const TSharedPtr<IPropertyHandle> TargetDependenciesPropertyHandle = TargetSettingsPropertyHandle->GetChildHandle(GET_MEMBER_NAME_CHECKED(FLocalizationTargetSettings, TargetDependencies));
+			if (TargetDependenciesPropertyHandle.IsValid() && TargetDependenciesPropertyHandle->IsValidHandle())
+			{
+				TargetDependenciesPropertyHandle->MarkHiddenByCustomization();
+				if (LocalizationTarget)
+				{
+					FDetailWidgetRow& TargetDependenciesRow = DetailCategoryBuilder.AddCustomRow( NSLOCTEXT("LocalizationTargetEditor", "ConflictReportStatusFilterString", "Conflict Report Status") );
+					TargetDependenciesRow.NameContent()
+						[
+							TargetDependenciesPropertyHandle->CreatePropertyNameWidget()
+						];
+					TargetDependenciesRow.ValueContent()
+						.HAlign(HAlign_Left)
+						.VAlign(VAlign_Center)
+						[
+							SNew(SComboButton)
+							.ButtonContent()
+							[
+								SAssignNew(TargetDependenciesHorizontalBox, SHorizontalBox)
+							]
+							.HasDownArrow(true)
+								.OnGetMenuContent_Lambda([this]() -> TSharedRef<SWidget>
+							{
+								RebuildTargetsList();
+
+								if (TargetDependenciesOptionsList.Num() > 0)
+								{
+									return SNew(SBox)
+										.MaxDesiredHeight(400.0f)
+										.MaxDesiredWidth(300.0f)
+										[
+											SNew(SListView<ULocalizationTarget*>)
+											.SelectionMode(ESelectionMode::None)
+											.ListItemsSource(&TargetDependenciesOptionsList)
+											.OnGenerateRow(this, &FLocalizationTargetDetailCustomization::OnGenerateTargetRow)
+										];
+								}
+								else
+								{
+									return SNullWidget::NullWidget;
+								}
+							})
+						];
+
+					RebuildTargetDependenciesBox();
+				}
+			}
+
+			const TSharedPtr<IPropertyHandle> AdditionalManifestDependenciesPropertyHandle = TargetSettingsPropertyHandle->GetChildHandle(GET_MEMBER_NAME_CHECKED(FLocalizationTargetSettings, AdditionalManifestDependencies));
+			if (AdditionalManifestDependenciesPropertyHandle.IsValid() && AdditionalManifestDependenciesPropertyHandle->IsValidHandle())
+			{
+				DetailCategoryBuilder.AddProperty(AdditionalManifestDependenciesPropertyHandle);
+			}
+		}
+
+		if (TargetSettingsPropertyHandle.IsValid() && TargetSettingsPropertyHandle->IsValidHandle())
+		{
+			IDetailCategoryBuilder& DetailCategoryBuilder = DetailBuilder.EditCategory(FName("Cultures"));
+
+			NativeCultureStatisticsPropertyHandle = TargetSettingsPropertyHandle->GetChildHandle(GET_MEMBER_NAME_CHECKED(FLocalizationTargetSettings, NativeCultureStatistics));
+			if (NativeCultureStatisticsPropertyHandle.IsValid() && NativeCultureStatisticsPropertyHandle->IsValidHandle())
+			{
+				NativeCultureNamePropertyHandle = NativeCultureStatisticsPropertyHandle->GetChildHandle(GET_MEMBER_NAME_CHECKED(FCultureStatistics, CultureName));
+				if (NativeCultureNamePropertyHandle.IsValid() && NativeCultureNamePropertyHandle->IsValidHandle())
+				{
+					NativeCultureNamePropertyHandle->MarkHiddenByCustomization();
+					FDetailWidgetRow& NativeCultureRow = DetailCategoryBuilder.AddCustomRow( NSLOCTEXT("LocalizationTargetEditor", "NativeCultureNameFilterString", "Native Culture Name") );
+
+					FString NativeCultureName;
+					NativeCultureNamePropertyHandle->GetValue(NativeCultureName);
+					FCulturePtr NativeCulture = FInternationalization::Get().GetCulture(NativeCultureName);
+
+					NativeCultureRow.NameContent()
+						[
+							NativeCultureNamePropertyHandle->CreatePropertyNameWidget(NSLOCTEXT("LocalizationTargetEditor", "NativeCultureNameLabel", "Native Culture"))
+						];
+
+					NativeCultureRow.ValueContent()
+						[
+							SAssignNew(NativeCultureComboButton, SComboButton)
+							.ButtonContent()
+							[
+								SNew(STextBlock)
+								.Text(this, &FLocalizationTargetDetailCustomization::GetNativeCultureDisplayName)
+								.ToolTipText(this, &FLocalizationTargetDetailCustomization::GetNativeCultureName)
+							]
+							.HasDownArrow(true)
+								.MenuContent()
+								[
+									SNew(SBox)
+									.MaxDesiredHeight(400.0f)
+									.MaxDesiredWidth(300.0f)
+									[
+										SNew(SCulturePicker)
+										.OnSelectionChanged(this, &FLocalizationTargetDetailCustomization::OnNativeCultureSelected)
+										.InitialSelection(NativeCulture)
+									]
+								]
+						];
+				}
+			}
+
+			SupportedCulturesStatisticsPropertyHandle = TargetSettingsPropertyHandle->GetChildHandle(GET_MEMBER_NAME_CHECKED(FLocalizationTargetSettings, SupportedCulturesStatistics));
+			if (SupportedCulturesStatisticsPropertyHandle.IsValid() && SupportedCulturesStatisticsPropertyHandle->IsValidHandle())
+			{
+				SupportedCulturesStatisticsPropertyHandle->MarkHiddenByCustomization();
+				SupportedCulturesStatisticsPropertyHandle_OnNumElementsChanged = FSimpleDelegate::CreateSP(this, &FLocalizationTargetDetailCustomization::RebuildListedCulturesList);
+				SupportedCulturesStatisticsPropertyHandle->AsArray()->SetOnNumElementsChanged(SupportedCulturesStatisticsPropertyHandle_OnNumElementsChanged);
+
+				FLocalizationTargetEditorCommands::Register();
+				const TSharedRef< FUICommandList > CommandList = MakeShareable(new FUICommandList);
+				FToolBarBuilder ToolBarBuilder( CommandList, FMultiBoxCustomization::AllowCustomization("LocalizationTargetEditor") );
+
+				CommandList->MapAction( FLocalizationTargetEditorCommands::Get().Gather, FExecuteAction::CreateSP(this, &FLocalizationTargetDetailCustomization::Gather));
+				ToolBarBuilder.AddToolBarButton(FLocalizationTargetEditorCommands::Get().Gather, NAME_None, TAttribute<FText>(), TAttribute<FText>(), FSlateIcon(FEditorStyle::GetStyleSetName(), "LocalizationTargetEditor.Gather"));
+
+				CommandList->MapAction( FLocalizationTargetEditorCommands::Get().ImportAllCultures, FExecuteAction::CreateSP(this, &FLocalizationTargetDetailCustomization::ImportAllCultures));
+				ToolBarBuilder.AddToolBarButton(FLocalizationTargetEditorCommands::Get().ImportAllCultures, NAME_None, TAttribute<FText>(), TAttribute<FText>(), FSlateIcon(FEditorStyle::GetStyleSetName(), "LocalizationTargetEditor.ImportForAllCultures"));
+
+				CommandList->MapAction( FLocalizationTargetEditorCommands::Get().ExportAllCultures, FExecuteAction::CreateSP(this, &FLocalizationTargetDetailCustomization::ExportAllCultures));
+				ToolBarBuilder.AddToolBarButton(FLocalizationTargetEditorCommands::Get().ExportAllCultures, NAME_None, TAttribute<FText>(), TAttribute<FText>(), FSlateIcon(FEditorStyle::GetStyleSetName(), "LocalizationTargetEditor.ExportForAllCultures"));
+
+				CommandList->MapAction( FLocalizationTargetEditorCommands::Get().RefreshWordCounts, FExecuteAction::CreateSP(this, &FLocalizationTargetDetailCustomization::RefreshWordCounts));
+				ToolBarBuilder.AddToolBarButton(FLocalizationTargetEditorCommands::Get().RefreshWordCounts, NAME_None, TAttribute<FText>(), TAttribute<FText>(), FSlateIcon(FEditorStyle::GetStyleSetName(), "LocalizationTargetEditor.RefreshWordCounts"));
+
+				BuildListedCulturesList();
+
+				DetailCategoryBuilder.AddCustomRow( NSLOCTEXT("LocalizationTargetEditor", "SupportedCultureNamesFilterString", "Supported Culture Names") )
+					.WholeRowContent()
+					[
+						SNew(SVerticalBox)
+						+SVerticalBox::Slot()
+						.AutoHeight()
+						[
+							ToolBarBuilder.MakeWidget()
+						]
+						+SVerticalBox::Slot()
+							.AutoHeight()
+							[
+								SAssignNew(SupportedCultureListView, SListView< TSharedPtr<IPropertyHandle> >)
+								.OnGenerateRow(this, &FLocalizationTargetDetailCustomization::OnGenerateCultureRow)
+								.ListItemsSource(&ListedCultureStatisticProperties)
+								.SelectionMode(ESelectionMode::None)
+								.HeaderRow
+								(
+								SNew(SHeaderRow)
+								+SHeaderRow::Column("Culture")
+								.DefaultLabel( NSLOCTEXT("LocalizationCulture", "CultureColumnLabel", "Culture"))
+								.HAlignHeader(HAlign_Fill)
+								.HAlignCell(HAlign_Fill)
+								.VAlignCell(VAlign_Center)
+								+SHeaderRow::Column("WordCount")
+								.DefaultLabel( NSLOCTEXT("LocalizationCulture", "WordCountColumnLabel", "Word Count"))
+								.HAlignHeader(HAlign_Center)
+								.HAlignCell(HAlign_Fill)
+								.VAlignCell(VAlign_Center)
+								+SHeaderRow::Column("Actions")
+								.DefaultLabel( NSLOCTEXT("LocalizationCulture", "ActionsColumnLabel", "Actions"))
+								.HAlignHeader(HAlign_Center)
+								.HAlignCell(HAlign_Center)
+								.VAlignCell(VAlign_Center)
+								)
+							]
+						+SVerticalBox::Slot()
+							.AutoHeight()
+							.VAlign(VAlign_Center)
+							[
+								SAssignNew(AddNewSupportedCultureComboButton, SComboButton)
+								.ButtonContent()
+								[
+									SNew(STextBlock)
+									.Text(NSLOCTEXT("LocalizationCulture", "AddNewCultureButtonLabel", "Add New Culture"))
+								]
+								.MenuContent()
+									[
+										SNew(SBox)
+										.MaxDesiredHeight(400.0f)
+										.MaxDesiredWidth(300.0f)
+										[
+											SAssignNew(SupportedCulturePicker, SCulturePicker)
+											.OnSelectionChanged(this, &FLocalizationTargetDetailCustomization::OnNewSupportedCultureSelected)
+											.IsCulturePickable(this, &FLocalizationTargetDetailCustomization::IsCultureSelectableAsSupported)
+										]
+									]
+							]
+					];
+			}
+		}
+	}
+
+	FLocalizationTargetSettings* FLocalizationTargetDetailCustomization::GetTargetSettings() const
+	{
+		return LocalizationTarget ? &(LocalizationTarget->Settings) : nullptr;
+	}
+
+	TSharedPtr<IPropertyHandle> FLocalizationTargetDetailCustomization::GetTargetSettingsPropertyHandle() const
+	{
+		return TargetSettingsPropertyHandle;
+	}
+
+	FText FLocalizationTargetDetailCustomization::GetTargetName() const
+	{
+		if (LocalizationTarget)
+		{
+			return FText::FromString(LocalizationTarget->Settings.Name);
+		}
+		return FText::GetEmpty();
+	}
+
+	bool FLocalizationTargetDetailCustomization::IsTargetNameUnique(const FString& Name) const
+	{
+		for (ULocalizationTarget* const TargetObject : ProjectSettings->TargetObjects)
+		{
+			if (TargetObject != LocalizationTarget)
+			{
+				if (TargetObject->Settings.Name == LocalizationTarget->Settings.Name)
+				{
+					return false;
+				}
+			}
+		}
+		return true;
+	}
+
+	void FLocalizationTargetDetailCustomization::OnTargetNameChanged(const FText& NewText)
+	{
+		const FString& NewName = NewText.ToString();
+
+		// TODO: Target names must be valid directory names, because they are used as directory names.
+		// ValidatePath allows /, which is not a valid directory name character
+		FText Error;
+		if (!FPaths::ValidatePath(NewName, &Error))
+		{
+			TargetNameEditableTextBox->SetError(Error);
+			return;		
+		}
+
+		// Target name must be unique.
+		if (!IsTargetNameUnique(NewName))
+		{
+			TargetNameEditableTextBox->SetError(NSLOCTEXT("LocalizationTargetEditor", "DuplicateTargetNameError", "Target name must be unique."));
+			return;
+		}
+
+		// Clear error if nothing has failed.
+		TargetNameEditableTextBox->SetError(FText::GetEmpty());
+	}
+
+	void FLocalizationTargetDetailCustomization::OnTargetNameCommitted(const FText& NewText, ETextCommit::Type Type)
+	{
+		// Target name must be unique.
+		if (!IsTargetNameUnique(NewText.ToString()))
+		{
+			return;
+		}
+
+		if (TargetSettingsPropertyHandle.IsValid() && TargetSettingsPropertyHandle->IsValidHandle())
+		{
+			FLocalizationTargetSettings* const TargetSettings = GetTargetSettings();
+			if (TargetSettings)
+			{
+				// Early out if the committed name is the same as the current name.
+				if (TargetSettings->Name == NewText.ToString())
+				{
+					return;
+				}
+
+				const TSharedPtr<IPropertyHandle> TargetNamePropertyHandle = TargetSettingsPropertyHandle->GetChildHandle(GET_MEMBER_NAME_CHECKED(FLocalizationTargetSettings, Name));
+
+				if (TargetNamePropertyHandle.IsValid() && TargetNamePropertyHandle->IsValidHandle())
+				{
+					TargetNamePropertyHandle->NotifyPreChange();
+				}
+
+				TargetSettings->Rename(NewText.ToString());
+
+				if (TargetNamePropertyHandle.IsValid() && TargetNamePropertyHandle->IsValidHandle())
+				{
+					TargetNamePropertyHandle->NotifyPostChange();
+				}
+			}
+		}
+	}
+
+	void FLocalizationTargetDetailCustomization::RebuildTargetDependenciesBox()
+	{
+		if (TargetDependenciesHorizontalBox.IsValid())
+		{
+			for (const TSharedPtr<SWidget>& Widget : TargetDependenciesWidgets)
+			{
+				TargetDependenciesHorizontalBox->RemoveSlot(Widget.ToSharedRef());
+			}
+			TargetDependenciesWidgets.Empty();
+
+			for (const FString& TargetDependency : LocalizationTarget->Settings.TargetDependencies)
+			{
+				const TSharedRef<SWidget> Widget = SNew(SBorder)
+					[
+						SNew(STextBlock)
+						.Text(FText::FromString(TargetDependency))
+					];
+
+				TargetDependenciesWidgets.Add(Widget);
+				TargetDependenciesHorizontalBox->AddSlot()
+					[
+						Widget
+					];
+			}
+		}
+	}
+
+	void FLocalizationTargetDetailCustomization::RebuildTargetsList()
+	{
+		TargetDependenciesOptionsList.Empty();
+		TFunction<bool (ULocalizationTarget* const)> DoesTargetDependOnUs;
+		DoesTargetDependOnUs = [&, this](ULocalizationTarget* const OtherTarget) -> bool
+		{
+			if (OtherTarget->Settings.TargetDependencies.Contains(LocalizationTarget->Settings.Name))
+			{
+				return true;
+			}
+
+			for (const FString& OtherTargetDependencyName : OtherTarget->Settings.TargetDependencies)
+			{
+				ULocalizationTarget** const OtherTargetDependency = ProjectSettings->TargetObjects.FindByPredicate([OtherTargetDependencyName](ULocalizationTarget* SomeLocalizationTarget)->bool{return SomeLocalizationTarget->Settings.Name == OtherTargetDependencyName;});
+				if (OtherTargetDependency && DoesTargetDependOnUs(*OtherTargetDependency))
+				{
+					return true;
+				}
+			}
+
+			return false;
+		};
+
+		for (ULocalizationTarget* const OtherTarget : ProjectSettings->TargetObjects)
+		{
+			if (OtherTarget != LocalizationTarget)
+			{
+				if (!DoesTargetDependOnUs(OtherTarget))
+				{
+					TargetDependenciesOptionsList.Add(OtherTarget);
+				}
+			}
+		}
+
+		if (TargetDependenciesListView.IsValid())
+		{
+			TargetDependenciesListView->RequestListRefresh();
+		}
+	}
+
+	TSharedRef<ITableRow> FLocalizationTargetDetailCustomization::OnGenerateTargetRow(ULocalizationTarget* OtherLocalizationTarget, const TSharedRef<STableViewBase>& Table)
+	{
+		return SNew(STableRow<ULocalizationTarget*>, Table)
+			.ShowSelection(true)
+			.Content()
+			[
+				SNew(SCheckBox)
+				.OnCheckStateChanged_Lambda([this, OtherLocalizationTarget](ECheckBoxState State)
+				{
+					OnTargetDependencyCheckStateChanged(OtherLocalizationTarget, State);
+				})
+					.IsChecked_Lambda([this, OtherLocalizationTarget]()->ECheckBoxState
+				{
+					return IsTargetDependencyChecked(OtherLocalizationTarget);
+				})
+					.Content()
+					[
+						SNew(STextBlock)
+						.Text(OtherLocalizationTarget->Settings.Name)
+					]
+			];
+	}
+
+	void FLocalizationTargetDetailCustomization::OnTargetDependencyCheckStateChanged(ULocalizationTarget* const OtherLocalizationTarget, const ECheckBoxState State)
+	{
+		const TSharedPtr<IPropertyHandle> TargetDependenciesPropertyHandle = TargetSettingsPropertyHandle->GetChildHandle(GET_MEMBER_NAME_CHECKED(FLocalizationTargetSettings, TargetDependencies));
+
+		if (TargetDependenciesPropertyHandle.IsValid() && TargetDependenciesPropertyHandle->IsValidHandle())
+		{
+			TargetDependenciesPropertyHandle->NotifyPreChange();
+		}
+
+		switch (State)
+		{
+		case ECheckBoxState::Checked:
+			LocalizationTarget->Settings.TargetDependencies.Add(OtherLocalizationTarget->Settings.Name);
+			break;
+		case ECheckBoxState::Unchecked:
+			LocalizationTarget->Settings.TargetDependencies.Remove(OtherLocalizationTarget->Settings.Name);
+			break;
+		}
+
+		if (TargetDependenciesPropertyHandle.IsValid() && TargetDependenciesPropertyHandle->IsValidHandle())
+		{
+			TargetDependenciesPropertyHandle->NotifyPostChange();
+		}
+
+		RebuildTargetDependenciesBox();
+	}
+
+	ECheckBoxState FLocalizationTargetDetailCustomization::IsTargetDependencyChecked(ULocalizationTarget* const OtherLocalizationTarget) const
+	{
+		return LocalizationTarget->Settings.TargetDependencies.Contains(OtherLocalizationTarget->Settings.Name) ? ECheckBoxState::Checked : ECheckBoxState::Unchecked;
+	}
+
+	FText FLocalizationTargetDetailCustomization::GetNativeCultureName() const
+	{
+		FCulturePtr NativeCulture;
+
+		if(NativeCultureNamePropertyHandle.IsValid() && NativeCultureNamePropertyHandle->IsValidHandle())
+		{
+			FString Value;
+			NativeCultureNamePropertyHandle->GetValue(Value);
+			NativeCulture = FInternationalization::Get().GetCulture(Value);
+		}
+
+		return NativeCulture.IsValid() ? FText::FromString(NativeCulture->GetName()) : FText::GetEmpty();
+	}
+
+	FText FLocalizationTargetDetailCustomization::GetNativeCultureDisplayName() const
+	{
+		FCulturePtr NativeCulture;
+
+		if(NativeCultureNamePropertyHandle.IsValid() && NativeCultureNamePropertyHandle->IsValidHandle())
+		{
+			FString Value;
+			NativeCultureNamePropertyHandle->GetValue(Value);
+			NativeCulture = FInternationalization::Get().GetCulture(Value);
+		}
+
+		return NativeCulture.IsValid() ? FText::FromString(NativeCulture->GetDisplayName()) : FText::GetEmpty();
+	}
+
+	void FLocalizationTargetDetailCustomization::OnNativeCultureSelected(FCulturePtr SelectedCulture, ESelectInfo::Type SelectInfo)
+	{
+		if(NativeCultureNamePropertyHandle.IsValid() && NativeCultureNamePropertyHandle->IsValidHandle())
+		{
+			NativeCultureNamePropertyHandle->SetValue(SelectedCulture->GetName());
+		}
+		if (NativeCultureComboButton.IsValid())
+		{
+			NativeCultureComboButton->SetIsOpen(false);
+		}
+
+		// TODO:	If the selected culture is an existing supported culture:
+		//			copy that supported culture's statistics to the native culture statistics
+		//			remove the supported culture's statistics
+
+		// TODO: Warn about this action invalidating all word counts and, realistically, invalidating all existing translations.
+		// Possibly offer an option to delete all archives?
+	}
+
+	void FLocalizationTargetDetailCustomization::Gather()
+	{
+		if (LocalizationTarget)
+		{
+			// Save unsaved packages.
+			bool DidPackagesNeedSaving;
+			const bool WerePackagesSaved = FEditorFileUtils::SaveDirtyPackages(true, true, true, false, false, &DidPackagesNeedSaving);
+
+			if (DidPackagesNeedSaving && !WerePackagesSaved)
+			{
+				// Give warning dialog.
+				const FText MessageText = NSLOCTEXT("LocalizationCultureActions", "UnsavedPackagesWarningDialogMessage", "There are unsaved changes. These changes may not be gathered from correctly.");
+				const FText TitleText = NSLOCTEXT("LocalizationCultureActions", "UnsavedPackagesWarningDialogTitle", "Unsaved Changes Before Gather");
+				switch(FMessageDialog::Open(EAppMsgType::OkCancel, MessageText, &TitleText))
+				{
+				case EAppReturnType::Cancel:
+					{
+						return;
+					}
+					break;
+				}
+			}
+
+			// Execute gather.
+			const TSharedPtr<SWindow> ParentWindow = FSlateApplication::Get().FindWidgetWindow(DetailLayoutBuilder->GetDetailsView().AsShared());
+			LocalizationCommandletTasks::GatherTarget(ParentWindow.ToSharedRef(), LocalizationTarget->Settings);
+
+			UpdateTargetFromReports();
+		}
+	}
+
+	void FLocalizationTargetDetailCustomization::ImportAllCultures()
+	{
+		if (LocalizationTarget)
+		{
+			const TSharedPtr<SWindow> ParentWindow = FSlateApplication::Get().FindWidgetWindow(DetailLayoutBuilder->GetDetailsView().AsShared());
+			LocalizationCommandletTasks::ImportTarget(ParentWindow.ToSharedRef(), LocalizationTarget->Settings);
+
+			UpdateTargetFromReports();
+		}
+	}
+
+	void FLocalizationTargetDetailCustomization::ExportAllCultures()
+	{
+		if (LocalizationTarget)
+		{
+			const TSharedPtr<SWindow> ParentWindow = FSlateApplication::Get().FindWidgetWindow(DetailLayoutBuilder->GetDetailsView().AsShared());
+			LocalizationCommandletTasks::ExportTarget(ParentWindow.ToSharedRef(), LocalizationTarget->Settings);
+		}
+	}
+
+	void FLocalizationTargetDetailCustomization::RefreshWordCounts()
+	{
+		if (LocalizationTarget)
+		{
+			const TSharedPtr<SWindow> ParentWindow = FSlateApplication::Get().FindWidgetWindow(DetailLayoutBuilder->GetDetailsView().AsShared());
+			LocalizationCommandletTasks::GenerateReportsForTarget(ParentWindow.ToSharedRef(), LocalizationTarget->Settings);
+
+			UpdateTargetFromReports();
+		}
+	}
+
+	void FLocalizationTargetDetailCustomization::UpdateTargetFromReports()
+	{
+		FLocalizationTargetSettings* const TargetSettings = GetTargetSettings();
+		if (TargetSettings)
+		{
+			TArray< TSharedPtr<IPropertyHandle> > WordCountPropertyHandles;
+
+			if (TargetSettingsPropertyHandle.IsValid() && TargetSettingsPropertyHandle->IsValidHandle())
+			{
+				const TSharedPtr<IPropertyHandle> NativeCultureStatisticsPropertyHandle = TargetSettingsPropertyHandle->GetChildHandle(GET_MEMBER_NAME_CHECKED(FLocalizationTargetSettings, NativeCultureStatistics));
+				if (NativeCultureStatisticsPropertyHandle.IsValid() && NativeCultureStatisticsPropertyHandle->IsValidHandle())
+				{
+					const TSharedPtr<IPropertyHandle> WordCountPropertyHandle = NativeCultureStatisticsPropertyHandle->GetChildHandle(GET_MEMBER_NAME_CHECKED(FCultureStatistics, WordCount));
+					if (WordCountPropertyHandle.IsValid() && WordCountPropertyHandle->IsValidHandle())
+					{
+						WordCountPropertyHandles.Add(WordCountPropertyHandle);
+					}
+				}
+
+				const TSharedPtr<IPropertyHandle> SupportedCulturesStatisticsPropertyHandle = TargetSettingsPropertyHandle->GetChildHandle(GET_MEMBER_NAME_CHECKED(FLocalizationTargetSettings, SupportedCulturesStatistics));
+				if (SupportedCulturesStatisticsPropertyHandle.IsValid() && SupportedCulturesStatisticsPropertyHandle->IsValidHandle())
+				{
+					uint32 SupportedCultureCount = 0;
+					SupportedCulturesStatisticsPropertyHandle->GetNumChildren(SupportedCultureCount);
+					for (uint32 i = 0; i < SupportedCultureCount; ++i)
+					{
+						const TSharedPtr<IPropertyHandle> ElementPropertyHandle = SupportedCulturesStatisticsPropertyHandle->GetChildHandle(i);
+						if (ElementPropertyHandle.IsValid() && ElementPropertyHandle->IsValidHandle())
+						{
+							const TSharedPtr<IPropertyHandle> WordCountPropertyHandle = SupportedCulturesStatisticsPropertyHandle->GetChildHandle(GET_MEMBER_NAME_CHECKED(FCultureStatistics, WordCount));
+							if (WordCountPropertyHandle.IsValid() && WordCountPropertyHandle->IsValidHandle())
+							{
+								WordCountPropertyHandles.Add(WordCountPropertyHandle);
+							}
+						}
+					}
+				}
+			}
+
+			for (const TSharedPtr<IPropertyHandle>& WordCountPropertyHandle : WordCountPropertyHandles)
+			{
+				WordCountPropertyHandle->NotifyPreChange();
+			}
+			TargetSettings->UpdateWordCountsFromCSV();
+			TargetSettings->UpdateStatusFromConflictReport();
+			for (const TSharedPtr<IPropertyHandle>& WordCountPropertyHandle : WordCountPropertyHandles)
+			{
+				WordCountPropertyHandle->NotifyPostChange();
+			}
+		}
+	}
+
+	void FLocalizationTargetDetailCustomization::BuildListedCulturesList()
+	{
+		ListedCultureStatisticProperties.AddUnique(NativeCultureStatisticsPropertyHandle);
+		const TSharedPtr<IPropertyHandleArray> SupportedCulturesStatisticsArrayPropertyHandle = SupportedCulturesStatisticsPropertyHandle->AsArray();
+		if (SupportedCulturesStatisticsArrayPropertyHandle.IsValid())
+		{
+			uint32 ElementCount = 0;
+			SupportedCulturesStatisticsArrayPropertyHandle->GetNumElements(ElementCount);
+			for (uint32 i = 0; i < ElementCount; ++i)
+			{
+				const TSharedPtr<IPropertyHandle> CultureStatisticsProperty = SupportedCulturesStatisticsArrayPropertyHandle->GetElement(i);
+				ListedCultureStatisticProperties.AddUnique(CultureStatisticsProperty);
+			}
+		}
+	}
+
+	void FLocalizationTargetDetailCustomization::RebuildListedCulturesList()
+	{
+		if (NewEntryIndexToBeInitialized != INDEX_NONE)
+		{
+			const TSharedPtr<IPropertyHandle> SupportedCultureStatisticsPropertyHandle = SupportedCulturesStatisticsPropertyHandle->GetChildHandle(NewEntryIndexToBeInitialized);
+
+			const TSharedPtr<IPropertyHandle> CultureNamePropertyHandle = SupportedCultureStatisticsPropertyHandle->GetChildHandle(GET_MEMBER_NAME_CHECKED(FCultureStatistics, CultureName));
+			if(CultureNamePropertyHandle.IsValid() && CultureNamePropertyHandle->IsValidHandle())
+			{
+				CultureNamePropertyHandle->SetValue(SelectedNewCulture->GetName());
+			}
+
+			const TSharedPtr<IPropertyHandle> WordCountPropertyHandle = SupportedCultureStatisticsPropertyHandle->GetChildHandle(GET_MEMBER_NAME_CHECKED(FCultureStatistics, WordCount));
+			if(WordCountPropertyHandle.IsValid() && WordCountPropertyHandle->IsValidHandle())
+			{
+				WordCountPropertyHandle->SetValue(0);
+			}
+
+			AddNewSupportedCultureComboButton->SetIsOpen(false);
+
+			NewEntryIndexToBeInitialized = INDEX_NONE;
+		}
+
+		ListedCultureStatisticProperties.Empty();
+		BuildListedCulturesList();
+
+		if (SupportedCultureListView.IsValid())
+		{
+			SupportedCultureListView->RequestListRefresh();
+		}
+	}
+
+	TSharedRef<ITableRow> FLocalizationTargetDetailCustomization::OnGenerateCultureRow(TSharedPtr<IPropertyHandle> CultureStatisticsPropertyHandle, const TSharedRef<STableViewBase>& Table)
+	{
+		return SNew(SLocalizationTargetEditorCultureRow, Table, DetailLayoutBuilder->GetPropertyUtilities(), TargetSettingsPropertyHandle.ToSharedRef(), CultureStatisticsPropertyHandle->GetIndexInArray());
+	}
+
+	bool FLocalizationTargetDetailCustomization::IsCultureSelectableAsSupported(FCulturePtr Culture)
+	{
+		// Can't select the native culture.
+		if (NativeCultureNamePropertyHandle.IsValid() && NativeCultureNamePropertyHandle->IsValidHandle())
+		{
+			FString NativeCultureName;
+			NativeCultureNamePropertyHandle->GetValue(NativeCultureName);
+			const FCulturePtr NativeCulture = FInternationalization::Get().GetCulture(NativeCultureName);
+			if (NativeCulture == Culture)
+			{
+				return false;
+			}
+		}
+
+		auto Is = [&](const TSharedPtr<IPropertyHandle>& SupportedCultureStatisticProperty)
+		{
+			// Can't select existing supported cultures.
+			const TSharedPtr<IPropertyHandle> SupportedCultureNamePropertyHandle = SupportedCultureStatisticProperty->GetChildHandle(GET_MEMBER_NAME_CHECKED(FCultureStatistics, CultureName));
+			if (SupportedCultureNamePropertyHandle.IsValid() && SupportedCultureNamePropertyHandle->IsValidHandle())
+			{
+				FString SupportedCultureName;
+				SupportedCultureNamePropertyHandle->GetValue(SupportedCultureName);
+				const FCulturePtr SupportedCulture = FInternationalization::Get().GetCulture(SupportedCultureName);
+				return SupportedCulture == Culture;
+			}
+
+			return false;
+		};
+
+		return !ListedCultureStatisticProperties.ContainsByPredicate(Is);
+	}
+
+	void FLocalizationTargetDetailCustomization::OnNewSupportedCultureSelected(FCulturePtr SelectedCulture, ESelectInfo::Type SelectInfo)
+	{
+		if(SupportedCulturesStatisticsPropertyHandle.IsValid() && SupportedCulturesStatisticsPropertyHandle->IsValidHandle())
+		{
+			uint32 NewElementIndex;
+			SupportedCulturesStatisticsPropertyHandle->AsArray()->GetNumElements(NewElementIndex);
+
+			// Add element, set info for later initialization.
+			SupportedCulturesStatisticsPropertyHandle->AsArray()->AddItem();
+			SelectedNewCulture = SelectedCulture;
+			NewEntryIndexToBeInitialized = NewElementIndex;
+
+			// Refresh UI.
+			if (SupportedCulturePicker.IsValid())
+			{
+				SupportedCulturePicker->RequestTreeRefresh();
+			}
+		}
+	}
+}
+
+void SLocalizationTargetEditor::Construct(const FArguments& InArgs, UProjectLocalizationSettings* const InProjectSettings, ULocalizationTarget* const InLocalizationTarget)
+{
+	check(InProjectSettings->TargetObjects.Contains(InLocalizationTarget));
+
+	FPropertyEditorModule& PropertyModule = FModuleManager::LoadModuleChecked<FPropertyEditorModule>("PropertyEditor");
+	FDetailsViewArgs DetailsViewArgs(false, false, false, FDetailsViewArgs::ENameAreaSettings::HideNameArea, false, nullptr, false, NAME_None);
+	TSharedRef<IDetailsView> DetailsView = PropertyModule.CreateDetailView(DetailsViewArgs);
+	DetailsView->RegisterInstancedCustomPropertyLayout(ULocalizationTarget::StaticClass(), FOnGetDetailCustomizationInstance::CreateLambda([&]()->TSharedRef<IDetailCustomization>{return MakeShareable(new FLocalizationTargetDetailCustomization(InProjectSettings, InLocalizationTarget));}));
+
+	ChildSlot
+		[
+			DetailsView
+		];
+
+	DetailsView->SetObject(InLocalizationTarget, true);
+}
+
+#undef LOCTEXT_NAMESPACE
\ No newline at end of file
diff --git a/Engine/Source/Editor/LocalizationDashboard/Private/SLocalizationTargetEditor.h b/Engine/Source/Editor/LocalizationDashboard/Private/SLocalizationTargetEditor.h
new file mode 100644
--- /dev/null
+++ b/Engine/Source/Editor/LocalizationDashboard/Private/SLocalizationTargetEditor.h
@@ -0,0 +1,18 @@
+// Copyright 1998-2014 Epic Games, Inc. All Rights Reserved.
+
+#pragma once
+
+#include "SCompoundWidget.h"
+
+class UProjectLocalizationSettings;
+class ULocalizationTarget;
+
+class SLocalizationTargetEditor : public SCompoundWidget
+{
+public:
+	SLATE_BEGIN_ARGS(SLocalizationTargetEditor) {}
+	SLATE_END_ARGS()
+
+	void Construct(const FArguments& InArgs, UProjectLocalizationSettings* const InProjectSettings, ULocalizationTarget* const InTarget);
+};
+
diff --git a/Engine/Source/Editor/LocalizationDashboard/Private/SLocalizationTargetEditorCultureRow.cpp b/Engine/Source/Editor/LocalizationDashboard/Private/SLocalizationTargetEditorCultureRow.cpp
new file mode 100644
--- /dev/null
+++ b/Engine/Source/Editor/LocalizationDashboard/Private/SLocalizationTargetEditorCultureRow.cpp
@@ -0,0 +1,387 @@
+// Copyright 1998-2014 Epic Games, Inc. All Rights Reserved.
+
+#include "LocalizationDashboardPrivatePCH.h"
+#include "SLocalizationTargetEditorCultureRow.h"
+#include "ITranslationEditor.h"
+#include "LocalizationConfigurationScript.h"
+#include "LocalizationCommandletTasks.h"
+
+void SLocalizationTargetEditorCultureRow::Construct(const FTableRowArgs& InArgs, const TSharedRef<STableViewBase>& OwnerTableView, const TSharedRef<IPropertyUtilities>& InPropertyUtilities, const TSharedRef<IPropertyHandle>& InTargetSettingsPropertyHandle, const int32 InCultureIndex)
+{
+	PropertyUtilities = InPropertyUtilities;
+	TargetSettingsPropertyHandle = InTargetSettingsPropertyHandle;
+	CultureIndex = InCultureIndex;
+
+	FSuperRowType::Construct(InArgs, OwnerTableView);
+}
+
+TSharedRef<SWidget> SLocalizationTargetEditorCultureRow::GenerateWidgetForColumn( const FName& ColumnName )
+{
+	TSharedPtr<SWidget> Return;
+
+	if (ColumnName == "Culture")
+	{
+		const auto& NativeCultureIconVisibilityLambda = [this]()
+		{
+			return IsNativeCultureForTarget() ? EVisibility::Visible : EVisibility::Hidden;
+		};
+		// Culture Name
+		return SNew(SHorizontalBox)
+			+SHorizontalBox::Slot()
+			.AutoWidth()
+			.HAlign(HAlign_Center)
+			.VAlign(VAlign_Center)
+			.Padding(4.0f)
+			[
+				SNew(SImage)
+				.Visibility_Lambda(NativeCultureIconVisibilityLambda)
+				.Image(FEditorStyle::GetBrush("LocalizationTargetEditor.NativeCulture"))
+			]
+		+SHorizontalBox::Slot()
+			.AutoWidth()
+			.HAlign(HAlign_Left)
+			.VAlign(VAlign_Center)
+			[
+				SNew(STextBlock)
+				.Text(this, &SLocalizationTargetEditorCultureRow::GetCultureDisplayName)
+				.ToolTipText(this, &SLocalizationTargetEditorCultureRow::GetCultureName)
+			];
+	}
+	else if(ColumnName == "WordCount")
+	{
+		// Progress Bar and Word Counts
+		Return = SNew(SOverlay)
+			+SOverlay::Slot()
+			.VAlign(VAlign_Fill)
+			.Padding(0.0f)
+			[
+				SNew(SProgressBar)
+				.Percent(this, &SLocalizationTargetEditorCultureRow::GetProgressPercentage)
+			]
+		+SOverlay::Slot()
+			.HAlign(HAlign_Center)
+			.VAlign(VAlign_Center)
+			.Padding(2.0f)
+			[
+				SNew(STextBlock)
+				.Text(this, &SLocalizationTargetEditorCultureRow::GetWordCountText)
+			];
+	}
+	else if(ColumnName == "Actions")
+	{
+		TSharedRef<SHorizontalBox> HorizontalBox = SNew(SHorizontalBox);
+		Return = HorizontalBox;
+
+		// Edit
+		HorizontalBox->AddSlot()
+			.FillWidth(1.0f)
+			.HAlign(HAlign_Center)
+			.VAlign(VAlign_Center)
+			[
+				SNew(SButton)
+				.ButtonStyle( FEditorStyle::Get(), TEXT("HoverHintOnly") )
+				.ToolTipText( NSLOCTEXT("LocalizationCultureActions", "EditButtonLabel", "Edit Translations") )
+				.OnClicked(this, &SLocalizationTargetEditorCultureRow::Edit)
+				.Content()
+				[
+					SNew(SImage)
+					.Image( FEditorStyle::GetBrush("LocalizationTargetEditor.EditTranslations") )
+				]
+			];
+
+		// Import
+		HorizontalBox->AddSlot()
+			.FillWidth(1.0f)
+			.HAlign(HAlign_Center)
+			.VAlign(VAlign_Center)
+			[
+				SNew(SButton)
+				.ButtonStyle( FEditorStyle::Get(), TEXT("HoverHintOnly") )
+				.ToolTipText( NSLOCTEXT("LocalizationCultureActions", "ImportButtonLabel", "Import") )
+				.OnClicked(this, &SLocalizationTargetEditorCultureRow::Import)
+				.Content()
+				[
+					SNew(SImage)
+					.Image( FEditorStyle::GetBrush("LocalizationTargetEditor.ImportForCulture") )
+				]
+			];
+
+		// Export
+		HorizontalBox->AddSlot()
+			.FillWidth(1.0f)
+			.HAlign(HAlign_Center)
+			.VAlign(VAlign_Center)
+			[
+				SNew(SButton)
+				.ButtonStyle( FEditorStyle::Get(), TEXT("HoverHintOnly") )
+				.ToolTipText(NSLOCTEXT("LocalizationTargetCultureActions", "ExportButtonLabel", "Export"))
+				.OnClicked(this, &SLocalizationTargetEditorCultureRow::Export)
+				.Content()
+				[
+					SNew(SImage)
+					.Image(FEditorStyle::GetBrush("LocalizationTargetEditor.ExportForCulture"))
+				]
+			];
+
+		// Delete Culture
+		HorizontalBox->AddSlot()
+			.FillWidth(1.0f)
+			.HAlign(HAlign_Center)
+			.VAlign(VAlign_Center)
+			[
+				SNew(SButton)
+				.ButtonStyle( FEditorStyle::Get(), TEXT("HoverHintOnly") )
+				.ToolTipText(NSLOCTEXT("LocalizationTargetActions", "DeleteButtonLabel", "Delete"))
+				.IsEnabled(this, &SLocalizationTargetEditorCultureRow::CanDelete)
+				.OnClicked(this, &SLocalizationTargetEditorCultureRow::EnqueueDeletion)
+				.Content()
+				[
+					SNew(SImage)
+					.Image(FEditorStyle::GetBrush("LocalizationTargetEditor.DeleteCulture"))
+				]
+			];
+	}
+
+	return Return.IsValid() ? Return.ToSharedRef() : SNullWidget::NullWidget;
+}
+
+FLocalizationTargetSettings* SLocalizationTargetEditorCultureRow::GetTargetSettings() const
+{
+	if (TargetSettingsPropertyHandle.IsValid() && TargetSettingsPropertyHandle->IsValidHandle())
+	{
+		TArray<void*> RawData;
+		TargetSettingsPropertyHandle->AccessRawData(RawData);
+		return RawData.Num() ? reinterpret_cast<FLocalizationTargetSettings*>(RawData.Top()) : nullptr;
+	}
+	return nullptr;
+}
+
+FCultureStatistics* SLocalizationTargetEditorCultureRow::GetCultureStatistics() const
+{
+	FLocalizationTargetSettings* const TargetSettings = GetTargetSettings();
+	if (TargetSettings)
+	{
+		if (CultureIndex == INDEX_NONE)
+		{
+			return &TargetSettings->NativeCultureStatistics;
+		}
+		else
+		{
+			return &TargetSettings->SupportedCulturesStatistics[CultureIndex];
+		}
+	}
+	return nullptr;
+}
+
+FCulturePtr SLocalizationTargetEditorCultureRow::GetCulture() const
+{
+	FCultureStatistics* const CultureStatistics = GetCultureStatistics();
+	if (CultureStatistics)
+	{
+		return FInternationalization::Get().GetCulture(CultureStatistics->CultureName);
+	}
+	return nullptr;
+}
+
+bool SLocalizationTargetEditorCultureRow::IsNativeCultureForTarget() const
+{
+	FLocalizationTargetSettings* const TargetSettings = GetTargetSettings();
+	return &(TargetSettings->NativeCultureStatistics) == GetCultureStatistics();
+}
+
+uint32 SLocalizationTargetEditorCultureRow::GetWordCount() const
+{
+	FCultureStatistics* const CultureStatistics = GetCultureStatistics();
+	if (CultureStatistics)
+	{
+		return CultureStatistics->WordCount;
+	}
+	return 0;
+}
+
+uint32 SLocalizationTargetEditorCultureRow::GetNativeWordCount() const
+{
+	FLocalizationTargetSettings* const TargetSettings = GetTargetSettings();
+	return TargetSettings ? TargetSettings->NativeCultureStatistics.WordCount : 0;
+}
+
+FText SLocalizationTargetEditorCultureRow::GetCultureDisplayName() const
+{
+	const FCulturePtr Culture = GetCulture();
+	return Culture.IsValid() ? FText::FromString(Culture->GetDisplayName()) : FText::GetEmpty();
+}
+
+FText SLocalizationTargetEditorCultureRow::GetCultureName() const
+{
+	const FCulturePtr Culture = GetCulture();
+	return Culture.IsValid() ? FText::FromString(Culture->GetName()) : FText::GetEmpty();
+}
+
+FText SLocalizationTargetEditorCultureRow::GetWordCountText() const
+{
+	FFormatNamedArguments Arguments;
+	Arguments.Add(TEXT("TranslatedWordCount"), FText::AsNumber(GetWordCount()));
+	TOptional<float> Percentage = GetProgressPercentage();
+	Arguments.Add(TEXT("TranslatedPercentage"), FText::AsPercent(Percentage.IsSet() ? Percentage.GetValue() : 0.0f));
+	return FText::Format(NSLOCTEXT("LocalizationDashboard", "CultureWordCountProgressFormat", "{TranslatedWordCount} ({TranslatedPercentage})"), Arguments);
+}
+
+TOptional<float> SLocalizationTargetEditorCultureRow::GetProgressPercentage() const
+{
+	uint32 WordCount = GetWordCount();
+	uint32 NativeWordCount = GetNativeWordCount();
+	return IsNativeCultureForTarget() ? 1.0f : NativeWordCount != 0 ? float(WordCount) / float(NativeWordCount) : 0.0f;
+}
+
+void SLocalizationTargetEditorCultureRow::UpdateTargetFromReports()
+{
+	FLocalizationTargetSettings* const TargetSettings = GetTargetSettings();
+	if (TargetSettings)
+	{
+		TArray< TSharedPtr<IPropertyHandle> > WordCountPropertyHandles;
+
+		if (TargetSettingsPropertyHandle.IsValid() && TargetSettingsPropertyHandle->IsValidHandle())
+		{
+			const TSharedPtr<IPropertyHandle> NativeCultureStatisticsPropertyHandle = TargetSettingsPropertyHandle->GetChildHandle(GET_MEMBER_NAME_CHECKED(FLocalizationTargetSettings, NativeCultureStatistics));
+			if (NativeCultureStatisticsPropertyHandle.IsValid() && NativeCultureStatisticsPropertyHandle->IsValidHandle())
+			{
+				const TSharedPtr<IPropertyHandle> WordCountPropertyHandle = NativeCultureStatisticsPropertyHandle->GetChildHandle(GET_MEMBER_NAME_CHECKED(FCultureStatistics, WordCount));
+				if (WordCountPropertyHandle.IsValid() && WordCountPropertyHandle->IsValidHandle())
+				{
+					WordCountPropertyHandles.Add(WordCountPropertyHandle);
+				}
+			}
+
+			const TSharedPtr<IPropertyHandle> SupportedCulturesStatisticsPropertyHandle = TargetSettingsPropertyHandle->GetChildHandle(GET_MEMBER_NAME_CHECKED(FLocalizationTargetSettings, SupportedCulturesStatistics));
+			if (SupportedCulturesStatisticsPropertyHandle.IsValid() && SupportedCulturesStatisticsPropertyHandle->IsValidHandle())
+			{
+				uint32 SupportedCultureCount = 0;
+				SupportedCulturesStatisticsPropertyHandle->GetNumChildren(SupportedCultureCount);
+				for (uint32 i = 0; i < SupportedCultureCount; ++i)
+				{
+					const TSharedPtr<IPropertyHandle> ElementPropertyHandle = SupportedCulturesStatisticsPropertyHandle->GetChildHandle(i);
+					if (ElementPropertyHandle.IsValid() && ElementPropertyHandle->IsValidHandle())
+					{
+						const TSharedPtr<IPropertyHandle> WordCountPropertyHandle = SupportedCulturesStatisticsPropertyHandle->GetChildHandle(GET_MEMBER_NAME_CHECKED(FCultureStatistics, WordCount));
+						if (WordCountPropertyHandle.IsValid() && WordCountPropertyHandle->IsValidHandle())
+						{
+							WordCountPropertyHandles.Add(WordCountPropertyHandle);
+						}
+					}
+				}
+			}
+		}
+
+		for (const TSharedPtr<IPropertyHandle>& WordCountPropertyHandle : WordCountPropertyHandles)
+		{
+			WordCountPropertyHandle->NotifyPreChange();
+		}
+		TargetSettings->UpdateWordCountsFromCSV();
+		TargetSettings->UpdateStatusFromConflictReport();
+		for (const TSharedPtr<IPropertyHandle>& WordCountPropertyHandle : WordCountPropertyHandles)
+		{
+			WordCountPropertyHandle->NotifyPostChange();
+		}
+	}
+}
+
+FReply SLocalizationTargetEditorCultureRow::Edit()
+{
+	const FCulturePtr Culture = GetCulture();
+
+	FLocalizationTargetSettings* const TargetSettings = GetTargetSettings();
+	if (TargetSettings && Culture.IsValid())
+	{
+		FModuleManager::LoadModuleChecked<ITranslationEditor>("TranslationEditor").OpenTranslationEditor(LocalizationConfigurationScript::GetManifestPath(*TargetSettings), LocalizationConfigurationScript::GetArchivePath(*TargetSettings, Culture->GetName()));
+	}
+
+	return FReply::Handled();
+}
+
+FReply SLocalizationTargetEditorCultureRow::Import()
+{
+	FLocalizationTargetSettings* const TargetSettings = GetTargetSettings();
+	if (TargetSettings)
+	{
+		const FCulturePtr Culture = GetCulture();
+		if (Culture.IsValid())
+		{
+			const TSharedPtr<SWindow> ParentWindow = FSlateApplication::Get().FindWidgetWindow(AsShared());
+			LocalizationCommandletTasks::ImportCulture(ParentWindow.ToSharedRef(), *TargetSettings, Culture->GetName());
+
+			UpdateTargetFromReports();
+		}
+	}
+
+
+	return FReply::Handled();
+}
+
+FReply SLocalizationTargetEditorCultureRow::Export()
+{
+	const FCulturePtr Culture = GetCulture();
+	FLocalizationTargetSettings* const TargetSettings = GetTargetSettings();
+	if (Culture.IsValid() && TargetSettings)
+	{
+		const TSharedPtr<SWindow> ParentWindow = FSlateApplication::Get().FindWidgetWindow(AsShared());
+		LocalizationCommandletTasks::ExportCulture(ParentWindow.ToSharedRef(), *TargetSettings, Culture->GetName());
+	}
+
+	return FReply::Handled();
+}
+
+bool SLocalizationTargetEditorCultureRow::CanDelete() const
+{
+	return !IsNativeCultureForTarget();
+}
+
+FReply SLocalizationTargetEditorCultureRow::EnqueueDeletion()
+{
+	PropertyUtilities->EnqueueDeferredAction(FSimpleDelegate::CreateSP(this, &SLocalizationTargetEditorCultureRow::Delete));
+	return FReply::Handled();
+}
+
+void SLocalizationTargetEditorCultureRow::Delete()
+{
+	static bool IsExecuting = false;
+	if (!IsExecuting)
+	{
+		TGuardValue<bool> ReentranceGuard(IsExecuting, true);
+
+		const FCulturePtr Culture = GetCulture();
+		FLocalizationTargetSettings* const TargetSettings = GetTargetSettings();
+		if (Culture.IsValid() && TargetSettings)
+		{
+			FText TitleText;
+			FText MessageText;
+
+			// Setup dialog for deleting supported culture.
+			const FText FormatPattern = NSLOCTEXT("LocalizationDashboard", "DeleteCultureConfirmationDialogMessage", "This action can not be undone and data will be permanently lost. Are you sure you would like to delete {CultureName}?");
+			FFormatNamedArguments Arguments;
+			Arguments.Add(TEXT("CultureName"), FText::FromString(Culture->GetDisplayName()));
+			MessageText = FText::Format(FormatPattern, Arguments);
+			TitleText = NSLOCTEXT("LocalizationDashboard", "DeleteCultureConfirmationDialogTitle", "Confirm Culture Deletion");
+
+			switch(FMessageDialog::Open(EAppMsgType::OkCancel, MessageText, &TitleText))
+			{
+			case EAppReturnType::Ok:
+				{
+					const FString CultureName = Culture->GetName();
+					TargetSettings->DeleteFiles(&CultureName);
+
+					// Remove this element from the parent array.
+					const TSharedPtr<IPropertyHandle> SupportedCulturesStatisticsPropertyHandle = TargetSettingsPropertyHandle.IsValid() && TargetSettingsPropertyHandle->IsValidHandle() ? TargetSettingsPropertyHandle->GetChildHandle(GET_MEMBER_NAME_CHECKED(FLocalizationTargetSettings, SupportedCulturesStatistics)) : nullptr;
+					if (SupportedCulturesStatisticsPropertyHandle.IsValid() && SupportedCulturesStatisticsPropertyHandle->IsValidHandle())
+					{
+						const TSharedPtr<IPropertyHandleArray> ArrayPropertyHandle = SupportedCulturesStatisticsPropertyHandle->AsArray();
+						if (ArrayPropertyHandle.IsValid())
+						{
+							ArrayPropertyHandle->DeleteItem(CultureIndex);
+						}
+					}
+				}
+				break;
+			}
+		}
+	}
+}
diff --git a/Engine/Source/Editor/LocalizationDashboard/Private/SLocalizationTargetEditorCultureRow.h b/Engine/Source/Editor/LocalizationDashboard/Private/SLocalizationTargetEditorCultureRow.h
new file mode 100644
--- /dev/null
+++ b/Engine/Source/Editor/LocalizationDashboard/Private/SLocalizationTargetEditorCultureRow.h
@@ -0,0 +1,43 @@
+// Copyright 1998-2014 Epic Games, Inc. All Rights Reserved.
+#pragma once
+
+#include "STableRow.h"
+#include "IPropertyUtilities.h"
+#include "PropertyHandle.h"
+#include "LocalizationTarget.h"
+
+class SLocalizationTargetEditorCultureRow : public SMultiColumnTableRow<FCulturePtr>
+{
+public:
+	void Construct(const FTableRowArgs& InArgs, const TSharedRef<STableViewBase>& OwnerTableView, const TSharedRef<IPropertyUtilities>& InPropertyUtilities, const TSharedRef<IPropertyHandle>& InTargetSettingsPropertyHandle, const int32 InCultureIndex);
+	TSharedRef<SWidget> GenerateWidgetForColumn( const FName& ColumnName ) override;
+
+private:
+	FLocalizationTargetSettings* GetTargetSettings() const;
+	FCultureStatistics* GetCultureStatistics() const;
+	FCulturePtr GetCulture() const;
+	bool IsNativeCultureForTarget() const;
+
+	uint32 GetWordCount() const;
+	uint32 GetNativeWordCount() const;
+
+	FText GetCultureDisplayName() const;
+	FText GetCultureName() const;
+
+	FText GetWordCountText() const;
+	TOptional<float> GetProgressPercentage() const;
+
+	void UpdateTargetFromReports();
+
+	FReply Edit();
+	FReply Import();
+	FReply Export();
+	bool CanDelete() const;
+	FReply EnqueueDeletion();
+	void Delete();
+
+private:
+	TSharedPtr<IPropertyUtilities> PropertyUtilities;
+	TSharedPtr<IPropertyHandle> TargetSettingsPropertyHandle;
+	int32 CultureIndex;
+};
\ No newline at end of file
diff --git a/Engine/Source/Editor/LocalizationDashboard/Private/SLocalizationTargetStatusButton.cpp b/Engine/Source/Editor/LocalizationDashboard/Private/SLocalizationTargetStatusButton.cpp
new file mode 100644
--- /dev/null
+++ b/Engine/Source/Editor/LocalizationDashboard/Private/SLocalizationTargetStatusButton.cpp
@@ -0,0 +1,124 @@
+// Copyright 1998-2014 Epic Games, Inc. All Rights Reserved.
+
+#include "LocalizationDashboardPrivatePCH.h"
+#include "SLocalizationTargetStatusButton.h"
+#include "LocalizationCommandletExecution.h"
+
+void SLocalizationTargetStatusButton::Construct(const FArguments& InArgs, ULocalizationTarget& InTarget)
+{
+	Target = &InTarget;
+
+	SButton::Construct(
+		SButton::FArguments()
+		.ButtonStyle( FEditorStyle::Get(), TEXT("HoverHintOnly") )
+		.OnClicked(this, &SLocalizationTargetStatusButton::OnClicked)
+		.ToolTipText(this, &SLocalizationTargetStatusButton::GetToolTipText)
+		);
+
+		ChildSlot
+			.HAlign(HAlign_Center)
+			.VAlign(VAlign_Center)
+			[
+				SNew(SImage)
+				.Image(this, &SLocalizationTargetStatusButton::GetImageBrush)
+				.ColorAndOpacity(this, &SLocalizationTargetStatusButton::GetColorAndOpacity)
+			];
+}
+
+const FSlateBrush* SLocalizationTargetStatusButton::GetImageBrush() const
+{
+	switch (Target->Settings.Status)
+	{
+	default:
+	case ELocalizationTargetStatus::Unknown:
+		return FEditorStyle::GetBrush("Icons.Warning");
+		break;
+	case ELocalizationTargetStatus::Clear:
+		return FEditorStyle::GetBrush("Symbols.Check");
+		break;
+	case ELocalizationTargetStatus::ConflictsPresent:
+		return FEditorStyle::GetBrush("Symbols.X");
+		break;
+	}
+}
+
+FSlateColor SLocalizationTargetStatusButton::GetColorAndOpacity() const
+{
+	switch (Target->Settings.Status)
+	{
+	default:
+	case ELocalizationTargetStatus::Unknown:
+		return FLinearColor::White;
+		break;
+	case ELocalizationTargetStatus::Clear:
+		return FLinearColor::Green;
+		break;
+	case ELocalizationTargetStatus::ConflictsPresent:
+		return FLinearColor::Red;
+		break;
+	}
+}
+
+FText SLocalizationTargetStatusButton::GetToolTipText() const
+{
+	switch (Target->Settings.Status)
+	{
+	default:
+	case ELocalizationTargetStatus::Unknown:
+		return NSLOCTEXT("LocalizationCulture", "StatusToolTip_Unknown", "Conflict report file not detected. Click to generate the conflict report.");
+		break;
+	case ELocalizationTargetStatus::Clear:
+		return NSLOCTEXT("LocalizationCulture", "StatusToolTip_Clear", "No conflicts detected.");
+		break;
+	case ELocalizationTargetStatus::ConflictsPresent:
+		return NSLOCTEXT("LocalizationCulture", "StatusToolTip_ConflictsPresent", "Conflicts detected. Click to open the conflict report.");
+		break;
+	}
+}
+
+FReply SLocalizationTargetStatusButton::OnClicked()
+{
+	switch (Target->Settings.Status)
+	{
+	case ELocalizationTargetStatus::Unknown:
+		{
+			// Generate conflict report.
+			const FString ReportScriptPath = LocalizationConfigurationScript::GetReportScriptPath(Target->Settings);
+			LocalizationConfigurationScript::GenerateReportScript(Target->Settings).Write(ReportScriptPath);
+			TSharedPtr<FLocalizationCommandletProcess> ReportProcess = FLocalizationCommandletProcess::Execute(ReportScriptPath);
+			if (ReportProcess.IsValid())
+			{
+				FPlatformProcess::WaitForProc(ReportProcess->GetHandle());
+				Target->Settings.UpdateStatusFromConflictReport();
+			}
+		}
+		break;
+	case ELocalizationTargetStatus::Clear:
+		// Do nothing.
+		break;
+	case ELocalizationTargetStatus::ConflictsPresent:
+		FString ReportString;
+		if (FFileHelper::LoadFileToString(ReportString, *LocalizationConfigurationScript::GetConflictReportPath(Target->Settings)))
+		{
+			TSharedPtr<SWindow> ReportWindow =
+				SNew(SWindow)
+				.Content()
+				[
+					SNew(SBorder)
+					.BorderImage(FCoreStyle::Get().GetBrush("ToolPanel.GroupBorder"))
+					[
+						SNew(SMultiLineEditableText)
+						.Text( FText::FromString(ReportString) )
+					]
+				];
+
+			if (ReportWindow.IsValid())
+			{
+				FSlateApplication::Get().AddWindow(ReportWindow.ToSharedRef());
+			}
+		}
+		break;
+	}
+
+	return FReply::Handled();
+}
\ No newline at end of file
diff --git a/Engine/Source/Editor/LocalizationDashboard/Private/SLocalizationTargetStatusButton.h b/Engine/Source/Editor/LocalizationDashboard/Private/SLocalizationTargetStatusButton.h
new file mode 100644
--- /dev/null
+++ b/Engine/Source/Editor/LocalizationDashboard/Private/SLocalizationTargetStatusButton.h
@@ -0,0 +1,24 @@
+// Copyright 1998-2014 Epic Games, Inc. All Rights Reserved.
+#pragma once
+
+#include "SCompoundWidget.h"
+#include "ProjectLocalizationSettings.h"
+#include "LocalizationConfigurationScript.h"
+
+class SLocalizationTargetStatusButton : public SButton
+{
+public:
+	SLATE_BEGIN_ARGS( SLocalizationTargetStatusButton ){}
+	SLATE_END_ARGS()
+
+	void Construct(const FArguments& InArgs, ULocalizationTarget& Target);
+
+private:
+	const FSlateBrush* GetImageBrush() const;
+	FSlateColor GetColorAndOpacity() const;
+	FText GetToolTipText() const;
+	FReply OnClicked();
+
+private:
+	ULocalizationTarget* Target;
+};
\ No newline at end of file
diff --git a/Engine/Source/Editor/LocalizationDashboard/Private/TargetsTableEntry.h b/Engine/Source/Editor/LocalizationDashboard/Private/TargetsTableEntry.h
new file mode 100644
--- /dev/null
+++ b/Engine/Source/Editor/LocalizationDashboard/Private/TargetsTableEntry.h
@@ -0,0 +1,30 @@
+// Copyright 1998-2014 Epic Games, Inc. All Rights Reserved.
+#pragma once
+
+class ULocalizationTarget;
+
+struct FTargetsTableEntry
+{
+	TSharedPtr<IPropertyHandle> TargetObjectPropertyHandle;
+	int32 CultureIndex;
+
+	FTargetsTableEntry()
+		: CultureIndex(INDEX_NONE)
+	{}
+
+	FTargetsTableEntry(const TSharedPtr<IPropertyHandle>& InTargetObjectPropertyHandle, const int32 InCultureIndex = INDEX_NONE)
+		: TargetObjectPropertyHandle(InTargetObjectPropertyHandle)
+		, CultureIndex(InCultureIndex)
+	{}
+
+	ULocalizationTarget* GetTarget()
+	{
+		if (TargetObjectPropertyHandle.IsValid() && TargetObjectPropertyHandle->IsValidHandle())
+		{
+			UObject* Value = nullptr;
+			TargetObjectPropertyHandle->GetValue(Value);
+			return Cast<ULocalizationTarget>(Value);
+		}
+		return nullptr;
+	}
+};
\ No newline at end of file
diff --git a/Engine/Source/Editor/LocalizationDashboard/Public/ILocalizationDashboardModule.h b/Engine/Source/Editor/LocalizationDashboard/Public/ILocalizationDashboardModule.h
new file mode 100644
--- /dev/null
+++ b/Engine/Source/Editor/LocalizationDashboard/Public/ILocalizationDashboardModule.h
@@ -0,0 +1,31 @@
+// Copyright 1998-2014 Epic Games, Inc. All Rights Reserved.
+
+#pragma once
+
+#include "ModuleInterface.h"
+
+class ILocalizationServiceProvider;
+
+/**
+ * Interface for localization dashboard module.
+ */
+class ILocalizationDashboardModule
+	: public IModuleInterface
+{
+public:
+	/** Virtual destructor. */
+	virtual ~ILocalizationDashboardModule() { }
+
+	/**
+	 * Shows the localization dashboard UI.
+	 */
+	virtual void Show() = 0;
+
+	virtual const TArray<ILocalizationServiceProvider*>& GetLocalizationServiceProviders() const = 0;
+	virtual ILocalizationServiceProvider* GetCurrentLocalizationServiceProvider() const = 0;
+
+	static ILocalizationDashboardModule& Get()
+	{
+		return FModuleManager::LoadModuleChecked<ILocalizationDashboardModule>("LocalizationDashboard");
+	}
+};
diff --git a/Engine/Source/Editor/LocalizationDashboard/Public/ILocalizationServiceProvider.h b/Engine/Source/Editor/LocalizationDashboard/Public/ILocalizationServiceProvider.h
new file mode 100644
--- /dev/null
+++ b/Engine/Source/Editor/LocalizationDashboard/Public/ILocalizationServiceProvider.h
@@ -0,0 +1,42 @@
+// Copyright 1998-2014 Epic Games, Inc. All Rights Reserved.
+
+#pragma once
+
+#include "Runtime/Core/Public/Features/IModularFeature.h"
+#include "EditorStyleSet.h"
+#include "PropertyHandle.h"
+#include "DetailCategoryBuilder.h"
+#include "ProjectLocalizationSettings.h"
+
+/**
+ * Interface to talking with localization service providers.
+ */
+class ILocalizationServiceProvider : public IModularFeature
+{
+public:
+	/**
+	 * Virtual destructor
+	 */
+	virtual ~ILocalizationServiceProvider() {}
+
+	virtual bool Initialize() = 0;
+	virtual void Terminate() = 0;
+
+	/** Get the localization service provider display name */
+	virtual const FText GetDisplayName() const = 0;
+
+	/** Get the localization service provider name */
+	virtual const FName GetName() const = 0;
+
+	/** 
+	 * Create a settings widget for display in the localization dashboard.
+	 * @returns a widget used to edit project-specific localization service settings.
+	 */
+	virtual void CustomizeSettingsDetails(IDetailCategoryBuilder& DetailCategoryBuilder) = 0;
+
+	/** 
+	 * Create a settings widget for display in the localization target editor.
+	 * @returns a widget used to edit target-specific localization service settings.
+	 */
+	virtual void CustomizeTargetDetails(IDetailCategoryBuilder& DetailCategoryBuilder, ULocalizationTarget& Target) = 0;
+};
\ No newline at end of file
diff --git a/Engine/Source/Editor/MainFrame/MainFrame.Build.cs b/Engine/Source/Editor/MainFrame/MainFrame.Build.cs
--- a/Engine/Source/Editor/MainFrame/MainFrame.Build.cs
+++ b/Engine/Source/Editor/MainFrame/MainFrame.Build.cs
@@ -1,4 +1,4 @@
-// Copyright 1998-2015 Epic Games, Inc. All Rights Reserved.
+﻿﻿// Copyright 1998-2015 Epic Games, Inc. All Rights Reserved.
 
 using UnrealBuildTool;
 
@@ -39,6 +39,7 @@ public class MainFrame : ModuleRules
 				"Projects",
                 "DeviceProfileEditor",
                 "UndoHistory",
+                "LocalizationDashboard",
 			}
 		);
 
@@ -54,6 +55,7 @@ public class MainFrame : ModuleRules
 				"Settings",
 				"SourceCodeAccess",
                 "Toolbox",
+                "LocalizationDashboard",
 			}
 		);
 
@@ -78,7 +80,8 @@ public class MainFrame : ModuleRules
                 "TranslationEditor",
 				"SourceCodeAccess",
 				"EditorLiveStreaming",
-				"HotReload"
+				"HotReload",
+                "LocalizationDashboard"
 			}
 		);
 	}
diff --git a/Engine/Source/Editor/MainFrame/Private/Frame/MainFrameActions.cpp b/Engine/Source/Editor/MainFrame/Private/Frame/MainFrameActions.cpp
--- a/Engine/Source/Editor/MainFrame/Private/Frame/MainFrameActions.cpp
+++ b/Engine/Source/Editor/MainFrame/Private/Frame/MainFrameActions.cpp
@@ -73,6 +73,9 @@ void FMainFrameCommands::RegisterCommands()
 	UI_COMMAND( PackagingSettings, "Packaging Settings...", "Opens the settings for project packaging", EUserInterfaceActionType::Button, FInputGesture() );
 	ActionList->MapAction( PackagingSettings, FExecuteAction::CreateStatic( &FMainFrameActionCallbacks::PackagingSettings ), DefaultExecuteAction );
 
+	//UI_COMMAND( LocalizeProject, "Localize Project...", "Opens the dashboard for managing project localization data.", EUserInterfaceActionType::Button, FInputGesture() );
+	//ActionList->MapAction( LocalizeProject, FExecuteAction::CreateStatic( &FMainFrameActionCallbacks::LocalizeProject ), DefaultExecuteAction );
+
 	const int32 MaxProjects = 20;
 	for( int32 CurProjectIndex = 0; CurProjectIndex < MaxProjects; ++CurProjectIndex )
 	{
@@ -640,6 +643,11 @@ void FMainFrameActionCallbacks::PackagingSettings()
 	FModuleManager::LoadModuleChecked<ISettingsModule>("Settings").ShowViewer("Project", "Project", "Packaging");
 }
 
+//void FMainFrameActionCallbacks::LocalizeProject()
+//{
+//	FModuleManager::LoadModuleChecked<ILocalizationDashboardModule>("LocalizationDashboard").Show();
+//}
+
 void FMainFrameActionCallbacks::SwitchProjectByIndex( int32 ProjectIndex )
 {
 	FUnrealEdMisc::Get().SwitchProject( ProjectNames[ ProjectIndex ] );
diff --git a/Engine/Source/Editor/MainFrame/Private/Frame/MainFrameActions.h b/Engine/Source/Editor/MainFrame/Private/Frame/MainFrameActions.h
--- a/Engine/Source/Editor/MainFrame/Private/Frame/MainFrameActions.h
+++ b/Engine/Source/Editor/MainFrame/Private/Frame/MainFrameActions.h
@@ -31,6 +31,7 @@ public:
 	TSharedPtr< FUICommandInfo > OpenIDE;
 	TSharedPtr< FUICommandInfo > RefreshCodeProject;
 	TSharedPtr< FUICommandInfo > PackagingSettings;
+	TSharedPtr< FUICommandInfo > LocalizeProject;
 	TArray< TSharedPtr< FUICommandInfo > > SwitchProjectCommands;
 
 	TSharedPtr< FUICommandInfo > OpenContentBrowser;
@@ -184,6 +185,9 @@ public:
 	/** Opens the Packaging settings tab */
 	static void PackagingSettings();
 
+	/** Opens the Project Localization Dashboard */
+	static void LocalizeProject();
+	
 	/** Restarts the editor and switches projects */
 	static void SwitchProjectByIndex( int32 ProjectIndex );
 
diff --git a/Engine/Source/Editor/MainFrame/Private/MainFramePrivatePCH.h b/Engine/Source/Editor/MainFrame/Private/MainFramePrivatePCH.h
--- a/Engine/Source/Editor/MainFrame/Private/MainFramePrivatePCH.h
+++ b/Engine/Source/Editor/MainFrame/Private/MainFramePrivatePCH.h
@@ -1,4 +1,4 @@
-// Copyright 1998-2015 Epic Games, Inc. All Rights Reserved.
+﻿// Copyright 1998-2015 Epic Games, Inc. All Rights Reserved.
 
 #pragma once
 
@@ -30,6 +30,7 @@
 #include "SourceCodeNavigation.h"
 #include "SourceControlWindows.h"
 #include "ISettingsModule.h"
+#include "ILocalizationDashboardModule.h"
 #include "TargetPlatform.h"
 #include "PlatformInfo.h"
 #include "UnrealEd.h"
diff --git a/Engine/Source/Editor/MainFrame/Private/Menus/MainMenu.cpp b/Engine/Source/Editor/MainFrame/Private/Menus/MainMenu.cpp
--- a/Engine/Source/Editor/MainFrame/Private/Menus/MainMenu.cpp
+++ b/Engine/Source/Editor/MainFrame/Private/Menus/MainMenu.cpp
@@ -189,10 +189,11 @@ void FMainMenu::FillWindowMenu( FMenuBuilder& MenuBuilder, const TSharedRef< FEx
 		// When the feature becomes permanent and need not check a flag, register a nomad spawner for it in the proper WorkspaceMenu category
 		bool bMessagingDebugger = GetDefault<UEditorExperimentalSettings>()->bMessagingDebugger;
 		bool bBlutility = GetDefault<UEditorExperimentalSettings>()->bEnableEditorUtilityBlueprints;
+		bool bLocalizationDashboard = GetDefault<UEditorExperimentalSettings>()->bEnableLocalizationDashboard;
 		bool bTranslationEditor = GetDefault<UEditorExperimentalSettings>()->bEnableTranslationEditor;
 
 		// Make sure at least one is enabled before creating the section
-		if (bMessagingDebugger || bBlutility || bTranslationEditor)
+		if (bMessagingDebugger || bBlutility || bLocalizationDashboard || bTranslationEditor)
 		{
 			MenuBuilder.BeginSection("ExperimentalTabSpawners", LOCTEXT("ExperimentalTabSpawnersHeading", "Experimental"));
 			{
@@ -218,6 +219,17 @@ void FMainMenu::FillWindowMenu( FMenuBuilder& MenuBuilder, const TSharedRef< FEx
 						);
 				}
 
+				// Localization Dashboard
+				if (bLocalizationDashboard)
+				{
+					MenuBuilder.AddMenuEntry(
+						LOCTEXT("LocalizationDashboardLabel", "Localization Dashboard"),
+						LOCTEXT("BlutilityShelfToolTip", "Open the Localization Dashboard for this Project."),
+						FSlateIcon(),
+						FUIAction(FExecuteAction::CreateStatic(&FMainMenu::OpenLocalizationDashboard))
+						);
+				}
+
 				// Translation Editor
 				if (bTranslationEditor)
 				{
@@ -363,10 +375,16 @@ TSharedRef< SWidget > FMainMenu::MakeMainTabMenu( const TSharedPtr<FTabManager>&
 
 				MenuBuilder.AddSubMenu(
 					LOCTEXT("PackageProjectSubMenuLabel", "Package Project"),
-					LOCTEXT("PackageProjectSubMenuToolTip", "Compile, cook and package your project and its content for distribution"),
+					LOCTEXT("PackageProjectSubMenuToolTip", "Compile, cook and package your project and its content for distribution."),
 					FNewMenuDelegate::CreateStatic( &FPackageProjectMenu::MakeMenu ), false, FSlateIcon(FEditorStyle::GetStyleSetName(), "MainFrame.PackageProject")
 				);
 
+				MenuBuilder.AddMenuEntry( FMainFrameCommands::Get().LocalizeProject,
+					NAME_None,
+					TAttribute<FText>(),
+					LOCTEXT("LocalizeProjectToolTip", "Gather text from your project and import/export translations.")
+					);
+
 				/*
 				MenuBuilder.AddSubMenu(
 					LOCTEXT("CookProjectSubMenuLabel", "Cook Project"),
diff --git a/Engine/Source/Editor/MainFrame/Private/Menus/MainMenu.h b/Engine/Source/Editor/MainFrame/Private/Menus/MainMenu.h
--- a/Engine/Source/Editor/MainFrame/Private/Menus/MainMenu.h
+++ b/Engine/Source/Editor/MainFrame/Private/Menus/MainMenu.h
@@ -84,6 +84,15 @@ private:
 		FGlobalTabmanager::Get()->InvokeTab(FName(TEXT("MessagingDebugger")));
 	}
 
+	/**
+	* Opens the experimental localization dashboard.
+	* Remove this when it is no longer experimental.
+	*/
+	static void OpenLocalizationDashboard()
+	{
+		FModuleManager::LoadModuleChecked<ILocalizationDashboardModule>("LocalizationDashboard").Show();
+	}
+
 	/**
 	* Opens the experimental blutility shelf tab.
 	* Remove this when it is no longer experimental.
diff --git a/Engine/Source/Editor/UnrealEd/Classes/Settings/EditorExperimentalSettings.h b/Engine/Source/Editor/UnrealEd/Classes/Settings/EditorExperimentalSettings.h
--- a/Engine/Source/Editor/UnrealEd/Classes/Settings/EditorExperimentalSettings.h
+++ b/Engine/Source/Editor/UnrealEd/Classes/Settings/EditorExperimentalSettings.h
@@ -21,6 +21,10 @@ class UNREALED_API UEditorExperimentalSettings
 
 public:
 
+	/** Allows usage of the Localization Dashboard */
+	UPROPERTY(EditAnywhere, config, Category = Tools, meta = (DisplayName = "Localization Dashboard"))
+	bool bEnableLocalizationDashboard;
+
 	/** Allows usage of the Translation Editor */
 	UPROPERTY(EditAnywhere, config, Category = Tools, meta = (DisplayName = "Translation Editor"))
 	bool bEnableTranslationEditor;
diff --git a/Engine/Source/Editor/UnrealEd/Private/Editor.cpp b/Engine/Source/Editor/UnrealEd/Private/Editor.cpp
--- a/Engine/Source/Editor/UnrealEd/Private/Editor.cpp
+++ b/Engine/Source/Editor/UnrealEd/Private/Editor.cpp
@@ -661,7 +661,8 @@ void UEditorEngine::Init(IEngineLoop* InEngineLoop)
 			TEXT("DeviceProfileEditor"),
 			TEXT("SourceCodeAccess"),
 			TEXT("BehaviorTreeEditor"),
-			TEXT("HardwareTargeting")
+			TEXT("HardwareTargeting"),
+			TEXT("LocalizationDashboard")
 		};
 
 		FScopedSlowTask SlowTask(ARRAY_COUNT(ModuleNames));
