Fixed build issues with new localization dashboard module files. Also replaced unnecessary usage of NSLOCTEXT with LOCTEXT_NAMESPACE defines and LOCTEXT.

diff --git a/Engine/Source/Editor/LocalizationDashboard/Private/LocalizationCommandletExecution.cpp b/Engine/Source/Editor/LocalizationDashboard/Private/LocalizationCommandletExecution.cpp
--- a/Engine/Source/Editor/LocalizationDashboard/Private/LocalizationCommandletExecution.cpp
+++ b/Engine/Source/Editor/LocalizationDashboard/Private/LocalizationCommandletExecution.cpp
@@ -8,13 +8,16 @@
 #include "DesktopPlatformModule.h"
 #include "SThrobber.h"
 #include "Commandlets/CommandletHelpers.h"
+#include "EditorStyleSet.h"
+
+#define LOCTEXT_NAMESPACE "LocalizationCommandletExecutor"
 
 namespace
 {
 	class SLocalizationCommandletExecutor : public SWindow
 	{
 		SLATE_BEGIN_ARGS(SLocalizationCommandletExecutor)
-			: _Title( NSLOCTEXT("LocalizationDashboard", "LocalizationCommandletExecutorTitle", "Localization Commandlet") )
+			: _Title( LOCTEXT("WindowTitle", "Localization Commandlet") )
 		{}
 		SLATE_ATTRIBUTE(FText, Title)
 			SLATE_END_ARGS()
@@ -163,7 +166,7 @@ namespace
 							.DefaultLabel(FText::GetEmpty())
 							.FixedWidth(20.0)
 							+ SHeaderRow::Column("TaskName")
-							.DefaultLabel(NSLOCTEXT("LocalizationCommandletExecutor", "TaskListNameColumnLabel", "Task"))
+							.DefaultLabel(LOCTEXT("TaskListNameColumnLabel", "Task"))
 							.FillWidth(1.0)
 							)
 							.ListItemsSource(&TaskListModels)
@@ -220,8 +223,8 @@ namespace
 						[
 							SNew(SButton)
 							.ContentPadding(FMargin(6.0f, 2.0f))
-							.Text(NSLOCTEXT("LocalizationCommandletExecutor", "CopyLogButtonText", "Copy Log"))
-							.ToolTipText(NSLOCTEXT("LocalizationCommandletExecutor", "CopyLogButtonTooltip", "Copy the logged text to the clipboard."))
+							.Text(LOCTEXT("CopyLogButtonText", "Copy Log"))
+							.ToolTipText(LOCTEXT("CopyLogButtonTooltip", "Copy the logged text to the clipboard."))
 							.OnClicked(this, &SLocalizationCommandletExecutor::OnCopyLogClicked)
 						]
 						+ SHorizontalBox::Slot()
@@ -230,8 +233,8 @@ namespace
 								SNew(SButton)
 								.ContentPadding(FMargin(6.0f, 2.0f))
 								.IsEnabled(false)
-								.Text(NSLOCTEXT("LocalizationCommandletExecutor", "SaveLogButtonText", "Save Log..."))
-								.ToolTipText(NSLOCTEXT("LocalizationCommandletExecutor", "SaveLogButtonToolTip", "Save the logged text to a file."))
+								.Text(LOCTEXT("SaveLogButtonText", "Save Log..."))
+								.ToolTipText(LOCTEXT("SaveLogButtonToolTip", "Save the logged text to a file."))
 								.OnClicked(this, &SLocalizationCommandletExecutor::OnSaveLogClicked)
 							]
 						+ SHorizontalBox::Slot()
@@ -585,7 +588,7 @@ namespace
 
 	FReply SLocalizationCommandletExecutor::OnSaveLogClicked()
 	{
-		const FString TextFileDescription = NSLOCTEXT("LocalizationCommandletExecutor", "TextFileDescription", "Text File").ToString();
+		const FString TextFileDescription = LOCTEXT("TextFileDescription", "Text File").ToString();
 		const FString TextFileExtension = TEXT("txt");
 		const FString TextFileExtensionWildcard = FString::Printf(TEXT("*.%s"), *TextFileExtension);
 		const FString FileTypes = FString::Printf(TEXT("%s (%s)|%s"), *TextFileDescription, *TextFileExtensionWildcard, *TextFileExtensionWildcard);
@@ -608,7 +611,7 @@ namespace
 
 			if (DesktopPlatform->SaveFileDialog(
 				ParentWindowWindowHandle,
-				NSLOCTEXT("LocalizationDashboard", "SaveLogDialogTitle", "Save Log to File").ToString(),
+				LOCTEXT("SaveLogDialogTitle", "Save Log to File").ToString(),
 				DefaultPath,
 				DefaultFilename,
 				FileTypes,
@@ -626,7 +629,7 @@ namespace
 
 	FText SLocalizationCommandletExecutor::GetCloseButtonText() const
 	{
-		return HasCompleted() ? NSLOCTEXT("LocalizationDashboard", "OkayButtonText", "Okay") : NSLOCTEXT("LocalizationDashboard", "CancelButtonText", "Cancel");
+		return HasCompleted() ? LOCTEXT("OkayButtonText", "Okay") : LOCTEXT("CancelButtonText", "Cancel");
 	}
 
 	FReply SLocalizationCommandletExecutor::OnCloseButtonClicked()
@@ -692,4 +695,6 @@ FLocalizationCommandletProcess::~FLocalizationCommandletProcess()
 		FPlatformProcess::TerminateProc(ProcessHandle);
 	}
 	FPlatformProcess::ClosePipe(ReadPipe, WritePipe);
-}
\ No newline at end of file
+}
+
+#undef LOCTEXT_NAMESPACE
\ No newline at end of file
diff --git a/Engine/Source/Editor/LocalizationDashboard/Private/LocalizationCommandletTasks.cpp b/Engine/Source/Editor/LocalizationDashboard/Private/LocalizationCommandletTasks.cpp
--- a/Engine/Source/Editor/LocalizationDashboard/Private/LocalizationCommandletTasks.cpp
+++ b/Engine/Source/Editor/LocalizationDashboard/Private/LocalizationCommandletTasks.cpp
@@ -5,6 +5,8 @@
 #include "LocalizationCommandletExecution.h"
 #include "LocalizationConfigurationScript.h"
 
+#define LOCTEXT_NAMESPACE "LocalizationCommandletTasks"
+
 bool LocalizationCommandletTasks::GatherTargets(const TSharedRef<SWindow>& ParentWindow, const TArray<FLocalizationTargetSettings*>& TargetsSettings)
 {
 	TArray<LocalizationCommandletExecution::FTask> Tasks;
@@ -14,18 +16,18 @@ bool LocalizationCommandletTasks::GatherTargets(const TSharedRef<SWindow>& Paren
 		FFormatNamedArguments Arguments;
 		Arguments.Add(TEXT("TargetName"), FText::FromString(TargetSettings->Name));
 
-		const FText GatherTaskName = FText::Format(NSLOCTEXT("LocalizationDashboard", "GatherTaskNameFormat", "Gather Text for {TargetName}"), Arguments);
+		const FText GatherTaskName = FText::Format(LOCTEXT("GatherTaskNameFormat", "Gather Text for {TargetName}"), Arguments);
 		const FString GatherScriptPath = LocalizationConfigurationScript::GetGatherScriptPath(*TargetSettings);
 		LocalizationConfigurationScript::GenerateGatherScript(*TargetSettings).Write(GatherScriptPath);
 		Tasks.Add(LocalizationCommandletExecution::FTask(GatherTaskName, GatherScriptPath));
 
-		const FText ReportTaskName = FText::Format(NSLOCTEXT("LocalizationDashboard", "ReportTaskNameFormat", "Generate Reports for {TargetName}"), Arguments);
+		const FText ReportTaskName = FText::Format(LOCTEXT("ReportTaskNameFormat", "Generate Reports for {TargetName}"), Arguments);
 		const FString ReportScriptPath = LocalizationConfigurationScript::GetReportScriptPath(*TargetSettings);
 		LocalizationConfigurationScript::GenerateReportScript(*TargetSettings).Write(ReportScriptPath);
 		Tasks.Add(LocalizationCommandletExecution::FTask(ReportTaskName, ReportScriptPath));
 	}
 
-	return LocalizationCommandletExecution::Execute(ParentWindow, NSLOCTEXT("LocalizationDashboard", "GatherAllTargetsWindowTitle", "Gather Text for All Targets"), Tasks);
+	return LocalizationCommandletExecution::Execute(ParentWindow, LOCTEXT("GatherAllTargetsWindowTitle", "Gather Text for All Targets"), Tasks);
 }
 
 bool LocalizationCommandletTasks::GatherTarget(const TSharedRef<SWindow>& ParentWindow, FLocalizationTargetSettings& TargetSettings)
@@ -34,16 +36,16 @@ bool LocalizationCommandletTasks::GatherTarget(const TSharedRef<SWindow>& Parent
 
 	const FString GatherScriptPath = LocalizationConfigurationScript::GetGatherScriptPath(TargetSettings);
 	LocalizationConfigurationScript::GenerateGatherScript(TargetSettings).Write(GatherScriptPath);
-	Tasks.Add(LocalizationCommandletExecution::FTask(NSLOCTEXT("LocalizationDashboard", "GatherTaskName", "Gather Text"), GatherScriptPath));
+	Tasks.Add(LocalizationCommandletExecution::FTask(LOCTEXT("GatherTaskName", "Gather Text"), GatherScriptPath));
 
 	const FString ReportScriptPath = LocalizationConfigurationScript::GetReportScriptPath(TargetSettings);
 	LocalizationConfigurationScript::GenerateReportScript(TargetSettings).Write(ReportScriptPath);
-	Tasks.Add(LocalizationCommandletExecution::FTask(NSLOCTEXT("LocalizationDashboard", "ReportTaskName", "Generate Reports"), ReportScriptPath));
+	Tasks.Add(LocalizationCommandletExecution::FTask(LOCTEXT("ReportTaskName", "Generate Reports"), ReportScriptPath));
 
 	FFormatNamedArguments Arguments;
 	Arguments.Add(TEXT("TargetName"), FText::FromString(TargetSettings.Name));
 
-	const FText WindowTitle = FText::Format(NSLOCTEXT("LocalizationDashboard", "GatherTargetWindowTitle", "Gather Text for Target {TargetName}"), Arguments);
+	const FText WindowTitle = FText::Format(LOCTEXT("GatherTargetWindowTitle", "Gather Text for Target {TargetName}"), Arguments);
 	return LocalizationCommandletExecution::Execute(ParentWindow, WindowTitle, Tasks);
 }
 
@@ -57,18 +59,18 @@ bool LocalizationCommandletTasks::ImportTargets(const TSharedRef<SWindow>& Paren
 		FFormatNamedArguments Arguments;
 		Arguments.Add(TEXT("TargetName"), FText::FromString(TargetSettings->Name));
 
-		const FText ImportTaskName = FText::Format(NSLOCTEXT("LocalizationDashboard", "ImportTaskNameFormat", "Import Translations for {TargetName}"), Arguments);
+		const FText ImportTaskName = FText::Format(LOCTEXT("ImportTaskNameFormat", "Import Translations for {TargetName}"), Arguments);
 		const FString ImportScriptPath = LocalizationConfigurationScript::GetImportScriptPath(*TargetSettings);
 		LocalizationConfigurationScript::GenerateImportScript(*TargetSettings).Write(ImportScriptPath);
 		Tasks.Add(LocalizationCommandletExecution::FTask(ImportTaskName, ImportScriptPath));
 
-		const FText ReportTaskName = FText::Format(NSLOCTEXT("LocalizationDashboard", "ReportTaskNameFormat", "Generate Reports for {TargetName}"), Arguments);
+		const FText ReportTaskName = FText::Format(LOCTEXT("ReportTaskNameFormat", "Generate Reports for {TargetName}"), Arguments);
 		const FString ReportScriptPath = LocalizationConfigurationScript::GetReportScriptPath(*TargetSettings);
 		LocalizationConfigurationScript::GenerateReportScript(*TargetSettings).Write(ReportScriptPath);
 		Tasks.Add(LocalizationCommandletExecution::FTask(ReportTaskName, ReportScriptPath));
 	}
 
-	return LocalizationCommandletExecution::Execute(ParentWindow, NSLOCTEXT("LocalizationDashboard", "ImportForAllTargetsWindowTitle", "Import Translations for All Targets"), Tasks);
+	return LocalizationCommandletExecution::Execute(ParentWindow, LOCTEXT("ImportForAllTargetsWindowTitle", "Import Translations for All Targets"), Tasks);
 }
 
 bool LocalizationCommandletTasks::ImportTarget(const TSharedRef<SWindow>& ParentWindow, FLocalizationTargetSettings& TargetSettings)
@@ -77,16 +79,16 @@ bool LocalizationCommandletTasks::ImportTarget(const TSharedRef<SWindow>& Parent
 
 	const FString ImportScriptPath = LocalizationConfigurationScript::GetImportScriptPath(TargetSettings);
 	LocalizationConfigurationScript::GenerateImportScript(TargetSettings).Write(ImportScriptPath);
-	Tasks.Add(LocalizationCommandletExecution::FTask(NSLOCTEXT("LocalizationDashboard", "ImportTaskName", "Import Translations"), ImportScriptPath));
+	Tasks.Add(LocalizationCommandletExecution::FTask(LOCTEXT("ImportTaskName", "Import Translations"), ImportScriptPath));
 
 	const FString ReportScriptPath = LocalizationConfigurationScript::GetReportScriptPath(TargetSettings);
 	LocalizationConfigurationScript::GenerateReportScript(TargetSettings).Write(ReportScriptPath);
-	Tasks.Add(LocalizationCommandletExecution::FTask(NSLOCTEXT("LocalizationDashboard", "ReportTaskName", "Generate Reports"), ReportScriptPath));
+	Tasks.Add(LocalizationCommandletExecution::FTask(LOCTEXT("ReportTaskName", "Generate Reports"), ReportScriptPath));
 
 	FFormatNamedArguments Arguments;
 	Arguments.Add(TEXT("TargetName"), FText::FromString(TargetSettings.Name));
 
-	const FText WindowTitle = FText::Format(NSLOCTEXT("LocalizationDashboard", "ImportForTargetWindowTitle", "Import Translations for Target {TargetName}"), Arguments);
+	const FText WindowTitle = FText::Format(LOCTEXT("ImportForTargetWindowTitle", "Import Translations for Target {TargetName}"), Arguments);
 	return LocalizationCommandletExecution::Execute(ParentWindow, WindowTitle, Tasks);
 }
 
@@ -102,17 +104,17 @@ bool LocalizationCommandletTasks::ImportCulture(const TSharedRef<SWindow>& Paren
 
 	const FString ImportScriptPath = LocalizationConfigurationScript::GetImportScriptPath(TargetSettings, &CultureName);
 	LocalizationConfigurationScript::GenerateImportScript(TargetSettings, &CultureName).Write(ImportScriptPath);
-	Tasks.Add(LocalizationCommandletExecution::FTask(NSLOCTEXT("LocalizationDashboard", "ImportTaskName", "Import Translations"), ImportScriptPath));
+	Tasks.Add(LocalizationCommandletExecution::FTask(LOCTEXT("ImportTaskName", "Import Translations"), ImportScriptPath));
 
 	const FString ReportScriptPath = LocalizationConfigurationScript::GetReportScriptPath(TargetSettings);
 	LocalizationConfigurationScript::GenerateReportScript(TargetSettings).Write(ReportScriptPath);
-	Tasks.Add(LocalizationCommandletExecution::FTask(NSLOCTEXT("LocalizationDashboard", "ReportTaskName", "Generate Reports"), ReportScriptPath));
+	Tasks.Add(LocalizationCommandletExecution::FTask(LOCTEXT("ReportTaskName", "Generate Reports"), ReportScriptPath));
 
 	FFormatNamedArguments Arguments;
 	Arguments.Add(TEXT("CultureName"), FText::FromString(Culture->GetDisplayName()));
 	Arguments.Add(TEXT("TargetName"), FText::FromString(TargetSettings.Name));
 
-	const FText WindowTitle = FText::Format(NSLOCTEXT("LocalizationDashboard", "ImportCultureForTargetWindowTitle", "Import {CultureName} Translations for Target {TargetName}"), Arguments);
+	const FText WindowTitle = FText::Format(LOCTEXT("ImportCultureForTargetWindowTitle", "Import {CultureName} Translations for Target {TargetName}"), Arguments);
 	return LocalizationCommandletExecution::Execute(ParentWindow, WindowTitle, Tasks);
 }
 
@@ -125,13 +127,13 @@ bool LocalizationCommandletTasks::ExportTargets(const TSharedRef<SWindow>& Paren
 		FFormatNamedArguments Arguments;
 		Arguments.Add(TEXT("TargetName"), FText::FromString(TargetSettings->Name));
 
-		const FText ExportTaskName = FText::Format(NSLOCTEXT("LocalizationDashboard", "ExportTaskNameFormat", "Export Translations for {TargetName}"), Arguments);
+		const FText ExportTaskName = FText::Format(LOCTEXT("ExportTaskNameFormat", "Export Translations for {TargetName}"), Arguments);
 		const FString ExportScriptPath = LocalizationConfigurationScript::GetExportScriptPath(*TargetSettings);
 		LocalizationConfigurationScript::GenerateExportScript(*TargetSettings).Write(ExportScriptPath);
 		Tasks.Add(LocalizationCommandletExecution::FTask(ExportTaskName, ExportScriptPath));
 	}
 
-	return LocalizationCommandletExecution::Execute(ParentWindow, NSLOCTEXT("LocalizationDashboard", "ExportForAllTargetsWindowTitle", "Export Translations for All Targets"), Tasks);
+	return LocalizationCommandletExecution::Execute(ParentWindow, LOCTEXT("ExportForAllTargetsWindowTitle", "Export Translations for All Targets"), Tasks);
 }
 
 bool LocalizationCommandletTasks::ExportTarget(const TSharedRef<SWindow>& ParentWindow, FLocalizationTargetSettings& TargetSettings)
@@ -140,12 +142,12 @@ bool LocalizationCommandletTasks::ExportTarget(const TSharedRef<SWindow>& Parent
 
 	const FString ExportScriptPath = LocalizationConfigurationScript::GetExportScriptPath(TargetSettings);
 	LocalizationConfigurationScript::GenerateExportScript(TargetSettings).Write(ExportScriptPath);
-	Tasks.Add(LocalizationCommandletExecution::FTask(NSLOCTEXT("LocalizationDashboard", "ExportTaskName", "Export Translations"), ExportScriptPath));
+	Tasks.Add(LocalizationCommandletExecution::FTask(LOCTEXT("ExportTaskName", "Export Translations"), ExportScriptPath));
 
 	FFormatNamedArguments Arguments;
 	Arguments.Add(TEXT("TargetName"), FText::FromString(TargetSettings.Name));
 
-	const FText WindowTitle = FText::Format(NSLOCTEXT("LocalizationDashboard", "ExportForTargetWindowTitle", "Export Translations for Target {TargetName}"), Arguments);
+	const FText WindowTitle = FText::Format(LOCTEXT("ExportForTargetWindowTitle", "Export Translations for Target {TargetName}"), Arguments);
 	return LocalizationCommandletExecution::Execute(ParentWindow, WindowTitle, Tasks);
 }
 
@@ -161,13 +163,13 @@ bool LocalizationCommandletTasks::ExportCulture(const TSharedRef<SWindow>& Paren
 
 	const FString ExportScriptPath = LocalizationConfigurationScript::GetExportScriptPath(TargetSettings, &CultureName);
 	LocalizationConfigurationScript::GenerateExportScript(TargetSettings, &CultureName).Write(ExportScriptPath);
-	Tasks.Add(LocalizationCommandletExecution::FTask(NSLOCTEXT("LocalizationDashboard", "ExportTaskName", "Export Translations"), ExportScriptPath));
+	Tasks.Add(LocalizationCommandletExecution::FTask(LOCTEXT("ExportTaskName", "Export Translations"), ExportScriptPath));
 
 	FFormatNamedArguments Arguments;
 	Arguments.Add(TEXT("CultureName"), FText::FromString(Culture->GetDisplayName()));
 	Arguments.Add(TEXT("TargetName"), FText::FromString(TargetSettings.Name));
 
-	const FText WindowTitle = FText::Format(NSLOCTEXT("LocalizationDashboard", "ExportCultureForTargetWindowTitle", "Export {CultureName} Translations for Target {TargetName}"), Arguments);
+	const FText WindowTitle = FText::Format(LOCTEXT("ExportCultureForTargetWindowTitle", "Export {CultureName} Translations for Target {TargetName}"), Arguments);
 	return LocalizationCommandletExecution::Execute(ParentWindow, WindowTitle, Tasks);
 }
 
@@ -180,13 +182,13 @@ bool LocalizationCommandletTasks::GenerateReportsForTargets(const TSharedRef<SWi
 		FFormatNamedArguments Arguments;
 		Arguments.Add(TEXT("TargetName"), FText::FromString(TargetSettings->Name));
 
-		const FText ReportTaskName = FText::Format(NSLOCTEXT("LocalizationDashboard", "ReportTaskNameFormat", "Generate Reports for {TargetName}"), Arguments);
+		const FText ReportTaskName = FText::Format(LOCTEXT("ReportTaskNameFormat", "Generate Reports for {TargetName}"), Arguments);
 		const FString ReportScriptPath = LocalizationConfigurationScript::GetReportScriptPath(*TargetSettings);
 		LocalizationConfigurationScript::GenerateReportScript(*TargetSettings).Write(ReportScriptPath);
 		Tasks.Add(LocalizationCommandletExecution::FTask(ReportTaskName, ReportScriptPath));
 	}
 
-	return LocalizationCommandletExecution::Execute(ParentWindow, NSLOCTEXT("LocalizationDashboard", "GenerateReportsForAllTargetsWindowTitle", "Generate Reports for All Targets"), Tasks);
+	return LocalizationCommandletExecution::Execute(ParentWindow, LOCTEXT("GenerateReportsForAllTargetsWindowTitle", "Generate Reports for All Targets"), Tasks);
 }
 
 bool LocalizationCommandletTasks::GenerateReportsForTarget(const TSharedRef<SWindow>& ParentWindow, FLocalizationTargetSettings& TargetSettings)
@@ -195,11 +197,13 @@ bool LocalizationCommandletTasks::GenerateReportsForTarget(const TSharedRef<SWin
 
 	const FString ReportScriptPath = LocalizationConfigurationScript::GetReportScriptPath(TargetSettings);
 	LocalizationConfigurationScript::GenerateReportScript(TargetSettings).Write(ReportScriptPath);
-	Tasks.Add(LocalizationCommandletExecution::FTask(NSLOCTEXT("LocalizationDashboard", "ReportTaskName", "Generate Reports"), ReportScriptPath));
+	Tasks.Add(LocalizationCommandletExecution::FTask(LOCTEXT("ReportTaskName", "Generate Reports"), ReportScriptPath));
 
 	FFormatNamedArguments Arguments;
 	Arguments.Add(TEXT("TargetName"), FText::FromString(TargetSettings.Name));
 
-	const FText WindowTitle = FText::Format(NSLOCTEXT("LocalizationDashboard", "GenerateReportsForTargetWindowTitle", "Generate Reports for Target {TargetName}"), Arguments);
+	const FText WindowTitle = FText::Format(LOCTEXT("GenerateReportsForTargetWindowTitle", "Generate Reports for Target {TargetName}"), Arguments);
 	return LocalizationCommandletExecution::Execute(ParentWindow, WindowTitle, Tasks);
-}
\ No newline at end of file
+}
+
+#undef LOCTEXT_NAMESPACE
\ No newline at end of file
diff --git a/Engine/Source/Editor/LocalizationDashboard/Private/LocalizationDashboardModule.cpp b/Engine/Source/Editor/LocalizationDashboard/Private/LocalizationDashboardModule.cpp
--- a/Engine/Source/Editor/LocalizationDashboard/Private/LocalizationDashboardModule.cpp
+++ b/Engine/Source/Editor/LocalizationDashboard/Private/LocalizationDashboardModule.cpp
@@ -8,6 +8,8 @@
 #include "Features/IModularFeatures.h"
 #include "ILocalizationServiceProvider.h"
 
+#define LOCTEXT_NAMESPACE "LocalizationDashboard"
+
 class FLocalizationDashboardModule
 	: public ILocalizationDashboardModule
 {
@@ -29,7 +31,7 @@ public:
 		const auto& SpawnMainTab = [this](const FSpawnTabArgs& Args) -> TSharedRef<SDockTab>
 		{
 			return SNew(SDockTab)
-				.Label(NSLOCTEXT("LocalizationDashboard", "MainTabTitle", "Localization Dashboard"))
+				.Label(LOCTEXT("MainTabTitle", "Localization Dashboard"))
 				.TabRole(ETabRole::MajorTab)
 				[
 					SNew(SLocalizationDashboard, Settings.Get())
@@ -37,7 +39,7 @@ public:
 		};
 
 		FGlobalTabmanager::Get()->RegisterTabSpawner(SLocalizationDashboard::TabName, FOnSpawnTab::CreateLambda( TFunction<TSharedRef<SDockTab> (const FSpawnTabArgs&)>(SpawnMainTab) ) )
-			.SetDisplayName(NSLOCTEXT("LocalizationDashboard", "MainTabTitle", "Localization Dashboard"))
+			.SetDisplayName(LOCTEXT("MainTabTitle", "Localization Dashboard"))
 			.SetMenuType(ETabSpawnerMenuType::Hidden);
 
 		ServiceProviders = IModularFeatures::Get().GetModularFeatureImplementations<ILocalizationServiceProvider>("LocalizationService");
@@ -80,4 +82,6 @@ private:
 	TArray<ILocalizationServiceProvider*> ServiceProviders;
 };
 
-IMPLEMENT_MODULE(FLocalizationDashboardModule, "LocalizationDashboard");
\ No newline at end of file
+IMPLEMENT_MODULE(FLocalizationDashboardModule, "LocalizationDashboard");
+
+#undef LOCTEXT_NAMESPACE
\ No newline at end of file
diff --git a/Engine/Source/Editor/LocalizationDashboard/Private/SCulturePicker.cpp b/Engine/Source/Editor/LocalizationDashboard/Private/SCulturePicker.cpp
--- a/Engine/Source/Editor/LocalizationDashboard/Private/SCulturePicker.cpp
+++ b/Engine/Source/Editor/LocalizationDashboard/Private/SCulturePicker.cpp
@@ -5,6 +5,8 @@
 #include "STableRow.h"
 #include "SSearchBox.h"
 
+#define LOCTEXT_NAMESPACE "CulturePicker"
+
 void SCulturePicker::Construct( const FArguments& InArgs )
 {
 	OnCultureSelectionChanged = InArgs._OnSelectionChanged;
@@ -37,7 +39,7 @@ void SCulturePicker::Construct( const FArguments& InArgs )
 			.AutoHeight()
 			[
 				SNew(SSearchBox)
-				.HintText( NSLOCTEXT("CulturePicker", "SearchHintText", "Name/Abbreviation") )
+				.HintText( LOCTEXT("SearchHintText", "Name/Abbreviation") )
 				.OnTextChanged(this, &SCulturePicker::OnFilterStringChanged)
 				.DelayChangeNotificationsWhileTyping(true)
 			]
@@ -268,4 +270,6 @@ void SCulturePicker::OnSelectionChanged(TSharedPtr<FCultureEntry> Entry, ESelect
 	{
 		OnCultureSelectionChanged.ExecuteIfBound( Entry->Culture, SelectInfo );
 	}
-}
\ No newline at end of file
+}
+
+#undef LOCTEXT_NAMESPACE
\ No newline at end of file
diff --git a/Engine/Source/Editor/LocalizationDashboard/Private/SLocalizationDashboard.cpp b/Engine/Source/Editor/LocalizationDashboard/Private/SLocalizationDashboard.cpp
--- a/Engine/Source/Editor/LocalizationDashboard/Private/SLocalizationDashboard.cpp
+++ b/Engine/Source/Editor/LocalizationDashboard/Private/SLocalizationDashboard.cpp
@@ -20,8 +20,9 @@
 #include "MultiBoxBuilder.h"
 #include "FeedbackContext.h"
 #include "LocalizationCommandletTasks.h"
+#include "FileHelpers.h"
 
-#define LOCTEXT_NAMSPACE "LocalizationDashboard"
+#define LOCTEXT_NAMESPACE "LocalizationDashboard"
 
 const FName SLocalizationDashboard::TabName("LocalizationDashboard");
 
@@ -83,7 +84,6 @@ namespace
 		, DetailLayoutBuilder(nullptr)
 		, ServiceProviderCategoryBuilder(nullptr)
 		, NewEntryIndexToBeInitialized(INDEX_NONE)
-		, IDetailCustomization()
 	{
 		TArray<ILocalizationServiceProvider*> ActualProviders = ILocalizationDashboardModule::Get().GetLocalizationServiceProviders();
 		for (ILocalizationServiceProvider* ActualProvider : ActualProviders)
@@ -146,7 +146,7 @@ namespace
 				ServiceProviderPropertyHandle->MarkHiddenByCustomization();
 				//ServiceProviderCategoryBuilder = &( DetailBuilder.EditCategory( FName("LocalizationServiceProvider") ) );
 
-				//FDetailWidgetRow& ServiceProviderDetailWidgetRow = ServiceProviderCategoryBuilder->AddCustomRow( NSLOCTEXT("LocalizationDashboard", "LocalizationServiceProviderFilterString", "Localization Service Provider") );
+				//FDetailWidgetRow& ServiceProviderDetailWidgetRow = ServiceProviderCategoryBuilder->AddCustomRow( LOCTEXT("LocalizationServiceProviderFilterString", "Localization Service Provider") );
 				//ServiceProviderDetailWidgetRow.NameContent()
 				//	[
 				//		ServiceProviderPropertyHandle->CreatePropertyNameWidget()
@@ -199,7 +199,7 @@ namespace
 
 				BuildTargetsList();
 
-				DetailCategoryBuilder.AddCustomRow( NSLOCTEXT("LocalizationDashboard", "TargetsFilterString", "Targets") )
+				DetailCategoryBuilder.AddCustomRow( LOCTEXT("TargetsFilterString", "Targets") )
 					.WholeRowContent()
 					[
 						SNew(SVerticalBox)
@@ -219,27 +219,27 @@ namespace
 								(
 								SNew(SHeaderRow)
 								+SHeaderRow::Column("Target")
-								.DefaultLabel( NSLOCTEXT("LocalizationDashboard", "TargetColumnLabel", "Target"))
+								.DefaultLabel( LOCTEXT("TargetColumnLabel", "Target"))
 								.HAlignHeader(HAlign_Left)
 								.HAlignCell(HAlign_Left)
 								.VAlignCell(VAlign_Center)
 								+SHeaderRow::Column("Status")
-								.DefaultLabel( NSLOCTEXT("LocalizationDashboard", "StatusColumnLabel", "Status"))
+								.DefaultLabel( LOCTEXT("StatusColumnLabel", "Status"))
 								.HAlignHeader(HAlign_Center)
 								.HAlignCell(HAlign_Center)
 								.VAlignCell(VAlign_Center)
 								+SHeaderRow::Column("Cultures")
-								.DefaultLabel( NSLOCTEXT("LocalizationDashboard", "CulturesColumnLabel", "Cultures"))
+								.DefaultLabel( LOCTEXT("CulturesColumnLabel", "Cultures"))
 								.HAlignHeader(HAlign_Fill)
 								.HAlignCell(HAlign_Fill)
 								.VAlignCell(VAlign_Center)
 								+SHeaderRow::Column("WordCount")
-								.DefaultLabel( NSLOCTEXT("LocalizationDashboard", "WordCountColumnLabel", "Word Count"))
+								.DefaultLabel( LOCTEXT("WordCountColumnLabel", "Word Count"))
 								.HAlignHeader(HAlign_Center)
 								.HAlignCell(HAlign_Center)
 								.VAlignCell(VAlign_Center)
 								+SHeaderRow::Column("Actions")
-								.DefaultLabel( NSLOCTEXT("LocalizationDashboard", "ActionsColumnLabel", "Actions"))
+								.DefaultLabel( LOCTEXT("ActionsColumnLabel", "Actions"))
 								.HAlignHeader(HAlign_Center)
 								.HAlignCell(HAlign_Center)
 								.VAlignCell(VAlign_Center)
@@ -250,7 +250,7 @@ namespace
 							.VAlign(VAlign_Center)
 							[
 								SNew(SButton)
-								.Text(NSLOCTEXT("LocalizationDashboard", "AddNewTargetButtonLabel", "Add New Target"))
+								.Text(LOCTEXT("AddNewTargetButtonLabel", "Add New Target"))
 								.OnClicked(this, &FProjectLocalizationSettingsDetailsCustomization::OnNewTargetButtonClicked)
 							]
 					];
@@ -321,14 +321,14 @@ namespace
 	FText FProjectLocalizationSettingsDetailsCustomization::GetCurrentServiceProviderDisplayName() const
 	{
 		ILocalizationServiceProvider* const LSP = ILocalizationDashboardModule::Get().GetCurrentLocalizationServiceProvider();
-		return LSP ? LSP->GetDisplayName() : NSLOCTEXT("LocalizationDashboard", "NoServiceProviderName", "None");
+		return LSP ? LSP->GetDisplayName() : LOCTEXT("NoServiceProviderName", "None");
 	}
 
 	TSharedRef<SWidget> FProjectLocalizationSettingsDetailsCustomization::ServiceProviderComboBox_OnGenerateWidget(TSharedPtr<FLocalizationServiceProviderWrapper> LSPWrapper) const
 	{
 		ILocalizationServiceProvider* const LSP = LSPWrapper->Provider;
 		return	SNew(STextBlock)
-			.Text(LSP ? LSP->GetDisplayName() : NSLOCTEXT("LocalizationDashboard", "NoServiceProviderName", "None"));
+			.Text(LSP ? LSP->GetDisplayName() : LOCTEXT("NoServiceProviderName", "None"));
 	}
 
 	void FProjectLocalizationSettingsDetailsCustomization::ServiceProviderComboBox_OnSelectionChanged(TSharedPtr<FLocalizationServiceProviderWrapper> LSPWrapper, ESelectInfo::Type SelectInfo)
diff --git a/Engine/Source/Editor/LocalizationDashboard/Private/SLocalizationDashboardTargetRow.cpp b/Engine/Source/Editor/LocalizationDashboard/Private/SLocalizationDashboardTargetRow.cpp
--- a/Engine/Source/Editor/LocalizationDashboard/Private/SLocalizationDashboardTargetRow.cpp
+++ b/Engine/Source/Editor/LocalizationDashboard/Private/SLocalizationDashboardTargetRow.cpp
@@ -9,6 +9,8 @@
 #include "LocalizationConfigurationScript.h"
 #include "LocalizationCommandletTasks.h"
 
+#define LOCTEXT_NAMESPACE "LocalizationDashboardTargetRow"
+
 TMap< TWeakObjectPtr<ULocalizationTarget>, TWeakPtr<SDockTab> > SLocalizationDashboardTargetRow::TargetToTabMap;
 
 void SLocalizationDashboardTargetRow::Construct(const FTableRowArgs& InArgs, const TSharedRef<STableViewBase>& OwnerTableView, const TSharedRef<IPropertyUtilities>& InPropertyUtilities, const TSharedRef<IPropertyHandle>& InTargetObjectPropertyHandle)
@@ -64,7 +66,7 @@ TSharedRef<SWidget> SLocalizationDashboardTargetRow::GenerateWidgetForColumn( co
 			[
 				SNew(SButton)
 				.ButtonStyle( FEditorStyle::Get(), TEXT("HoverHintOnly") )
-				.ToolTipText( NSLOCTEXT("LocalizationTargetActions", "GatherButtonLabel", "Gather") )
+				.ToolTipText( LOCTEXT("GatherButtonLabel", "Gather") )
 				.OnClicked(this, &SLocalizationDashboardTargetRow::Gather)
 				.Content()
 				[
@@ -81,7 +83,7 @@ TSharedRef<SWidget> SLocalizationDashboardTargetRow::GenerateWidgetForColumn( co
 			[
 				SNew(SButton)
 				.ButtonStyle( FEditorStyle::Get(), TEXT("HoverHintOnly") )
-				.ToolTipText(NSLOCTEXT("LocalizationTargetActions", "RefreshWordCountButtonLabel", "Refresh"))
+				.ToolTipText( LOCTEXT("RefreshWordCountButtonLabel", "Refresh") )
 				.OnClicked(this, &SLocalizationDashboardTargetRow::RefreshWordCount)
 				.Content()
 				[
@@ -98,7 +100,7 @@ TSharedRef<SWidget> SLocalizationDashboardTargetRow::GenerateWidgetForColumn( co
 			[
 				SNew(SButton)
 				.ButtonStyle( FEditorStyle::Get(), TEXT("HoverHintOnly") )
-				.ToolTipText( NSLOCTEXT("LocalizationTargetActions", "ImportAllButtonLabel", "Import All") )
+				.ToolTipText( LOCTEXT("ImportAllButtonLabel", "Import All") )
 				.OnClicked(this, &SLocalizationDashboardTargetRow::ImportAll)
 				.Content()
 				[
@@ -115,7 +117,7 @@ TSharedRef<SWidget> SLocalizationDashboardTargetRow::GenerateWidgetForColumn( co
 			[
 				SNew(SButton)
 				.ButtonStyle( FEditorStyle::Get(), TEXT("HoverHintOnly") )
-				.ToolTipText(NSLOCTEXT("LocalizationTargetActions", "ExportAllButtonLabel", "Export All"))
+				.ToolTipText(LOCTEXT("ExportAllButtonLabel", "Export All"))
 				.OnClicked(this, &SLocalizationDashboardTargetRow::ExportAll)
 				.Content()
 				[
@@ -132,7 +134,7 @@ TSharedRef<SWidget> SLocalizationDashboardTargetRow::GenerateWidgetForColumn( co
 			[
 				SNew(SButton)
 				.ButtonStyle( FEditorStyle::Get(), TEXT("HoverHintOnly") )
-				.ToolTipText(NSLOCTEXT("LocalizationTargetActions", "DeleteButtonLabel", "Delete"))
+				.ToolTipText(LOCTEXT("DeleteButtonLabel", "Delete"))
 				.OnClicked(this, &SLocalizationDashboardTargetRow::EnqueueDeletion)
 				.Content()
 				[
@@ -448,3 +450,5 @@ void SLocalizationDashboardTargetRow::Delete()
 		}
 	}
 }
+
+#undef LOCTEXT_NAMESPACE
\ No newline at end of file
diff --git a/Engine/Source/Editor/LocalizationDashboard/Private/SLocalizationTargetEditor.cpp b/Engine/Source/Editor/LocalizationDashboard/Private/SLocalizationTargetEditor.cpp
--- a/Engine/Source/Editor/LocalizationDashboard/Private/SLocalizationTargetEditor.cpp
+++ b/Engine/Source/Editor/LocalizationDashboard/Private/SLocalizationTargetEditor.cpp
@@ -19,7 +19,7 @@
 #include "SLocalizationTargetStatusButton.h"
 #include "LocalizationCommandletTasks.h"
 
-#define LOCTEXT_NAMSPACE "LocalizationCulture"
+#define LOCTEXT_NAMESPACE "LocalizationTargetEditor"
 
 namespace
 {
@@ -152,7 +152,7 @@ namespace
 			const TSharedPtr<IPropertyHandle> NamePropertyHandle = TargetSettingsPropertyHandle->GetChildHandle(GET_MEMBER_NAME_CHECKED(FLocalizationTargetSettings, Name));
 			if (NamePropertyHandle.IsValid() && NamePropertyHandle->IsValidHandle())
 			{
-				FDetailWidgetRow& DetailWidgetRow = DetailCategoryBuilder.AddCustomRow( NSLOCTEXT("LocalizationTargetEditor", "TargetNameFilterString", "Target Name") );
+				FDetailWidgetRow& DetailWidgetRow = DetailCategoryBuilder.AddCustomRow( LOCTEXT("TargetNameFilterString", "Target Name") );
 				DetailWidgetRow.NameContent()
 					[
 						NamePropertyHandle->CreatePropertyNameWidget()
@@ -172,7 +172,7 @@ namespace
 				StatusPropertyHandle->MarkHiddenByCustomization();
 				if (LocalizationTarget)
 				{
-					FDetailWidgetRow& StatusRow = DetailCategoryBuilder.AddCustomRow( NSLOCTEXT("LocalizationTargetEditor", "ConflictReportStatusFilterString", "Conflict Report Status") );
+					FDetailWidgetRow& StatusRow = DetailCategoryBuilder.AddCustomRow( LOCTEXT("ConflictReportStatusFilterString", "Conflict Report Status") );
 					StatusRow.NameContent()
 						[
 							StatusPropertyHandle->CreatePropertyNameWidget()
@@ -192,7 +192,7 @@ namespace
 				TargetDependenciesPropertyHandle->MarkHiddenByCustomization();
 				if (LocalizationTarget)
 				{
-					FDetailWidgetRow& TargetDependenciesRow = DetailCategoryBuilder.AddCustomRow( NSLOCTEXT("LocalizationTargetEditor", "ConflictReportStatusFilterString", "Conflict Report Status") );
+					FDetailWidgetRow& TargetDependenciesRow = DetailCategoryBuilder.AddCustomRow( LOCTEXT("ConflictReportStatusFilterString", "Conflict Report Status") );
 					TargetDependenciesRow.NameContent()
 						[
 							TargetDependenciesPropertyHandle->CreatePropertyNameWidget()
@@ -252,7 +252,7 @@ namespace
 				if (NativeCultureNamePropertyHandle.IsValid() && NativeCultureNamePropertyHandle->IsValidHandle())
 				{
 					NativeCultureNamePropertyHandle->MarkHiddenByCustomization();
-					FDetailWidgetRow& NativeCultureRow = DetailCategoryBuilder.AddCustomRow( NSLOCTEXT("LocalizationTargetEditor", "NativeCultureNameFilterString", "Native Culture Name") );
+					FDetailWidgetRow& NativeCultureRow = DetailCategoryBuilder.AddCustomRow( LOCTEXT("NativeCultureNameFilterString", "Native Culture Name") );
 
 					FString NativeCultureName;
 					NativeCultureNamePropertyHandle->GetValue(NativeCultureName);
@@ -260,7 +260,7 @@ namespace
 
 					NativeCultureRow.NameContent()
 						[
-							NativeCultureNamePropertyHandle->CreatePropertyNameWidget(NSLOCTEXT("LocalizationTargetEditor", "NativeCultureNameLabel", "Native Culture"))
+							NativeCultureNamePropertyHandle->CreatePropertyNameWidget(LOCTEXT("NativeCultureNameLabel", "Native Culture"))
 						];
 
 					NativeCultureRow.ValueContent()
@@ -313,7 +313,7 @@ namespace
 
 				BuildListedCulturesList();
 
-				DetailCategoryBuilder.AddCustomRow( NSLOCTEXT("LocalizationTargetEditor", "SupportedCultureNamesFilterString", "Supported Culture Names") )
+				DetailCategoryBuilder.AddCustomRow( LOCTEXT("SupportedCultureNamesFilterString", "Supported Culture Names") )
 					.WholeRowContent()
 					[
 						SNew(SVerticalBox)
@@ -426,7 +426,7 @@ namespace
 		// Target name must be unique.
 		if (!IsTargetNameUnique(NewName))
 		{
-			TargetNameEditableTextBox->SetError(NSLOCTEXT("LocalizationTargetEditor", "DuplicateTargetNameError", "Target name must be unique."));
+			TargetNameEditableTextBox->SetError(LOCTEXT("DuplicateTargetNameError", "Target name must be unique."));
 			return;
 		}
 
diff --git a/Engine/Source/Editor/LocalizationDashboard/Private/SLocalizationTargetEditorCultureRow.cpp b/Engine/Source/Editor/LocalizationDashboard/Private/SLocalizationTargetEditorCultureRow.cpp
--- a/Engine/Source/Editor/LocalizationDashboard/Private/SLocalizationTargetEditorCultureRow.cpp
+++ b/Engine/Source/Editor/LocalizationDashboard/Private/SLocalizationTargetEditorCultureRow.cpp
@@ -6,6 +6,8 @@
 #include "LocalizationConfigurationScript.h"
 #include "LocalizationCommandletTasks.h"
 
+#define LOCTEXT_NAMESPACE "LocalizationTargetEditorCultureRow"
+
 void SLocalizationTargetEditorCultureRow::Construct(const FTableRowArgs& InArgs, const TSharedRef<STableViewBase>& OwnerTableView, const TSharedRef<IPropertyUtilities>& InPropertyUtilities, const TSharedRef<IPropertyHandle>& InTargetSettingsPropertyHandle, const int32 InCultureIndex)
 {
 	PropertyUtilities = InPropertyUtilities;
@@ -80,7 +82,7 @@ TSharedRef<SWidget> SLocalizationTargetEditorCultureRow::GenerateWidgetForColumn
 			[
 				SNew(SButton)
 				.ButtonStyle( FEditorStyle::Get(), TEXT("HoverHintOnly") )
-				.ToolTipText( NSLOCTEXT("LocalizationCultureActions", "EditButtonLabel", "Edit Translations") )
+				.ToolTipText( LOCTEXT("EditButtonLabel", "Edit Translations") )
 				.OnClicked(this, &SLocalizationTargetEditorCultureRow::Edit)
 				.Content()
 				[
@@ -97,7 +99,7 @@ TSharedRef<SWidget> SLocalizationTargetEditorCultureRow::GenerateWidgetForColumn
 			[
 				SNew(SButton)
 				.ButtonStyle( FEditorStyle::Get(), TEXT("HoverHintOnly") )
-				.ToolTipText( NSLOCTEXT("LocalizationCultureActions", "ImportButtonLabel", "Import") )
+				.ToolTipText( LOCTEXT("ImportButtonLabel", "Import") )
 				.OnClicked(this, &SLocalizationTargetEditorCultureRow::Import)
 				.Content()
 				[
@@ -223,7 +225,7 @@ FText SLocalizationTargetEditorCultureRow::GetWordCountText() const
 	Arguments.Add(TEXT("TranslatedWordCount"), FText::AsNumber(GetWordCount()));
 	TOptional<float> Percentage = GetProgressPercentage();
 	Arguments.Add(TEXT("TranslatedPercentage"), FText::AsPercent(Percentage.IsSet() ? Percentage.GetValue() : 0.0f));
-	return FText::Format(NSLOCTEXT("LocalizationDashboard", "CultureWordCountProgressFormat", "{TranslatedWordCount} ({TranslatedPercentage})"), Arguments);
+	return FText::Format(LOCTEXT("CultureWordCountProgressFormat", "{TranslatedWordCount} ({TranslatedPercentage})"), Arguments);
 }
 
 TOptional<float> SLocalizationTargetEditorCultureRow::GetProgressPercentage() const
@@ -356,11 +358,11 @@ void SLocalizationTargetEditorCultureRow::Delete()
 			FText MessageText;
 
 			// Setup dialog for deleting supported culture.
-			const FText FormatPattern = NSLOCTEXT("LocalizationDashboard", "DeleteCultureConfirmationDialogMessage", "This action can not be undone and data will be permanently lost. Are you sure you would like to delete {CultureName}?");
+			const FText FormatPattern = LOCTEXT("DeleteCultureConfirmationDialogMessage", "This action can not be undone and data will be permanently lost. Are you sure you would like to delete {CultureName}?");
 			FFormatNamedArguments Arguments;
 			Arguments.Add(TEXT("CultureName"), FText::FromString(Culture->GetDisplayName()));
 			MessageText = FText::Format(FormatPattern, Arguments);
-			TitleText = NSLOCTEXT("LocalizationDashboard", "DeleteCultureConfirmationDialogTitle", "Confirm Culture Deletion");
+			TitleText = LOCTEXT("DeleteCultureConfirmationDialogTitle", "Confirm Culture Deletion");
 
 			switch(FMessageDialog::Open(EAppMsgType::OkCancel, MessageText, &TitleText))
 			{
@@ -385,3 +387,5 @@ void SLocalizationTargetEditorCultureRow::Delete()
 		}
 	}
 }
+
+#undef LOCTEXT_NAMESPACE
\ No newline at end of file
diff --git a/Engine/Source/Editor/LocalizationDashboard/Private/SLocalizationTargetStatusButton.cpp b/Engine/Source/Editor/LocalizationDashboard/Private/SLocalizationTargetStatusButton.cpp
--- a/Engine/Source/Editor/LocalizationDashboard/Private/SLocalizationTargetStatusButton.cpp
+++ b/Engine/Source/Editor/LocalizationDashboard/Private/SLocalizationTargetStatusButton.cpp
@@ -4,6 +4,8 @@
 #include "SLocalizationTargetStatusButton.h"
 #include "LocalizationCommandletExecution.h"
 
+#define LOCTEXT_NAMESPACE "LocalizationTargetStatusButton"
+
 void SLocalizationTargetStatusButton::Construct(const FArguments& InArgs, ULocalizationTarget& InTarget)
 {
 	Target = &InTarget;
@@ -65,13 +67,13 @@ FText SLocalizationTargetStatusButton::GetToolTipText() const
 	{
 	default:
 	case ELocalizationTargetStatus::Unknown:
-		return NSLOCTEXT("LocalizationCulture", "StatusToolTip_Unknown", "Conflict report file not detected. Click to generate the conflict report.");
+		return LOCTEXT("StatusToolTip_Unknown", "Conflict report file not detected. Click to generate the conflict report.");
 		break;
 	case ELocalizationTargetStatus::Clear:
-		return NSLOCTEXT("LocalizationCulture", "StatusToolTip_Clear", "No conflicts detected.");
+		return LOCTEXT("StatusToolTip_Clear", "No conflicts detected.");
 		break;
 	case ELocalizationTargetStatus::ConflictsPresent:
-		return NSLOCTEXT("LocalizationCulture", "StatusToolTip_ConflictsPresent", "Conflicts detected. Click to open the conflict report.");
+		return LOCTEXT("StatusToolTip_ConflictsPresent", "Conflicts detected. Click to open the conflict report.");
 		break;
 	}
 }
@@ -121,4 +123,6 @@ FReply SLocalizationTargetStatusButton::OnClicked()
 	}
 
 	return FReply::Handled();
-}
\ No newline at end of file
+}
+
+#undef LOCTEXT_NAMESPACE
\ No newline at end of file
