Redesigned layout for localization dashboard and target editor tabs.

diff --git a/Engine/Source/Editor/LocalizationDashboard/Private/LocalizationDashboard.cpp b/Engine/Source/Editor/LocalizationDashboard/Private/LocalizationDashboard.cpp
new file mode 100644
--- /dev/null
+++ b/Engine/Source/Editor/LocalizationDashboard/Private/LocalizationDashboard.cpp
@@ -0,0 +1,662 @@
+// Copyright 1998-2014 Epic Games, Inc. All Rights Reserved.
+
+#include "LocalizationDashboardPrivatePCH.h"
+#include "LocalizationDashboard.h"
+#include "IDetailCustomization.h"
+#include "EditorStyleSet.h"
+#include "DetailLayoutBuilder.h"
+#include "ProjectLocalizationSettings.h"
+#include "PropertyEditorModule.h"
+#include "IDetailsView.h"
+#include "DetailWidgetRow.h"
+#include "DetailCategoryBuilder.h"
+#include "PropertyHandle.h"
+#include "SLocalizationDashboardTargetRow.h"
+#include "TargetsTableEntry.h"
+#include "Runtime/Core/Public/Features/IModularFeatures.h"
+#include "ILocalizationDashboardModule.h"
+#include "ILocalizationServiceProvider.h"
+#include "IPropertyUtilities.h"
+#include "MultiBoxBuilder.h"
+#include "FeedbackContext.h"
+#include "LocalizationCommandletTasks.h"
+#include "FileHelpers.h"
+#include "SLocalizationTargetEditor.h"
+
+#define LOCTEXT_NAMESPACE "LocalizationDashboard"
+
+const FName FLocalizationDashboard::TabName("LocalizationDashboard");
+
+class ILocalizationServiceProvider;
+
+namespace
+{
+	struct FLocalizationServiceProviderWrapper;
+
+	class FProjectLocalizationSettingsDetailsCustomization : public IDetailCustomization
+	{
+	public:
+		FProjectLocalizationSettingsDetailsCustomization(UProjectLocalizationSettings* const InSettings);
+		void CustomizeDetails(IDetailLayoutBuilder& DetailBuilder) override;
+
+	private:
+		void BuildTargetsList();
+		void RebuildTargetsList();
+
+		FText GetCurrentServiceProviderDisplayName() const;
+		TSharedRef<SWidget> ServiceProviderComboBox_OnGenerateWidget(TSharedPtr<FLocalizationServiceProviderWrapper> LSPWrapper) const;
+		void ServiceProviderComboBox_OnSelectionChanged(TSharedPtr<FLocalizationServiceProviderWrapper> LSPWrapper, ESelectInfo::Type SelectInfo);
+
+		void GatherAllTargets();
+		void ImportAllTargets();
+		void ExportAllTargets();
+		void UpdateTargetFromReports(ULocalizationTarget* const LocalizationTarget);
+
+		TSharedRef<ITableRow> OnGenerateRow(TSharedPtr<IPropertyHandle> TargetObjectPropertyHandle, const TSharedRef<STableViewBase>& Table);
+		FReply OnNewTargetButtonClicked();
+
+	private:
+		UProjectLocalizationSettings* Settings;
+
+		IDetailLayoutBuilder* DetailLayoutBuilder;
+
+		TSharedPtr<IPropertyHandle> ServiceProviderPropertyHandle;
+		IDetailCategoryBuilder* ServiceProviderCategoryBuilder;
+		TArray< TSharedPtr<FLocalizationServiceProviderWrapper> > Providers;
+
+		TSharedPtr<IPropertyHandle> TargetObjectsPropertyHandle;
+		FSimpleDelegate TargetsArrayPropertyHandle_OnNumElementsChanged;
+		TArray< TSharedPtr<IPropertyHandle> > TargetsList;
+		TSharedPtr< SListView< TSharedPtr<IPropertyHandle> > > TargetsListView;
+
+		/* If set, the entry at the index specified needs to be initialized as soon as possible. */
+		int32 NewEntryIndexToBeInitialized;
+	};
+	struct FLocalizationServiceProviderWrapper
+	{
+		FLocalizationServiceProviderWrapper() : Provider(nullptr) {}
+		FLocalizationServiceProviderWrapper(ILocalizationServiceProvider* const InProvider) : Provider(InProvider) {}
+
+		ILocalizationServiceProvider* Provider;
+	};
+
+	FProjectLocalizationSettingsDetailsCustomization::FProjectLocalizationSettingsDetailsCustomization(UProjectLocalizationSettings* const InSettings)
+		: Settings(InSettings)
+		, DetailLayoutBuilder(nullptr)
+		, ServiceProviderCategoryBuilder(nullptr)
+		, NewEntryIndexToBeInitialized(INDEX_NONE)
+	{
+		TArray<ILocalizationServiceProvider*> ActualProviders = ILocalizationDashboardModule::Get().GetLocalizationServiceProviders();
+		for (ILocalizationServiceProvider* ActualProvider : ActualProviders)
+		{
+			TSharedPtr<FLocalizationServiceProviderWrapper> Provider = MakeShareable(new FLocalizationServiceProviderWrapper(ActualProvider));
+			Providers.Add(Provider);
+		}
+	}
+
+	class FindProviderPredicate
+	{
+	public:
+		FindProviderPredicate(ILocalizationServiceProvider* const InActualProvider)
+			: ActualProvider(InActualProvider)
+		{
+		}
+
+		bool operator()(const TSharedPtr<FLocalizationServiceProviderWrapper>& Provider)
+		{
+			return Provider->Provider == ActualProvider;
+		}
+
+	private:
+		ILocalizationServiceProvider* ActualProvider;
+	};
+
+
+	class FLocalizationDashboardCommands : public TCommands<FLocalizationDashboardCommands>
+	{
+	public:
+		FLocalizationDashboardCommands() 
+			: TCommands<FLocalizationDashboardCommands>("LocalizationDashboard", NSLOCTEXT("Contexts", "LocalizationDashboard", "Localization Dashboard"), NAME_None, FEditorStyle::GetStyleSetName())
+		{
+		}
+
+		TSharedPtr<FUICommandInfo> GatherAllTargets;
+		TSharedPtr<FUICommandInfo> ImportAllTargets;
+		TSharedPtr<FUICommandInfo> ExportAllTargets;
+
+		/** Initialize commands */
+		virtual void RegisterCommands() override;
+	};
+
+	void FLocalizationDashboardCommands::RegisterCommands() 
+	{
+		UI_COMMAND( GatherAllTargets, "Gather All", "Gathers text for all targets in the project.", EUserInterfaceActionType::Button, FInputGesture() );
+		UI_COMMAND( ImportAllTargets, "Import All", "Imports translations for all cultures of all targets in the project.", EUserInterfaceActionType::Button, FInputGesture() );
+		UI_COMMAND( ExportAllTargets, "Export All", "Exports translations for all cultures of all targets in the project.", EUserInterfaceActionType::Button, FInputGesture() );
+	}
+
+	void FProjectLocalizationSettingsDetailsCustomization::CustomizeDetails(IDetailLayoutBuilder& DetailBuilder)
+	{
+		DetailLayoutBuilder = &DetailBuilder;
+
+		{
+			ServiceProviderPropertyHandle = DetailBuilder.GetProperty(GET_MEMBER_NAME_CHECKED(UProjectLocalizationSettings, LocalizationServiceProvider));
+
+			if (ServiceProviderPropertyHandle.IsValid() && ServiceProviderPropertyHandle->IsValidHandle())
+			{
+				ServiceProviderPropertyHandle->MarkHiddenByCustomization();
+				//ServiceProviderCategoryBuilder = &( DetailBuilder.EditCategory( FName("LocalizationServiceProvider") ) );
+
+				//FDetailWidgetRow& ServiceProviderDetailWidgetRow = ServiceProviderCategoryBuilder->AddCustomRow( LOCTEXT("LocalizationServiceProviderFilterString", "Localization Service Provider") );
+				//ServiceProviderDetailWidgetRow.NameContent()
+				//	[
+				//		ServiceProviderPropertyHandle->CreatePropertyNameWidget()
+				//	];
+
+				//ILocalizationServiceProvider* const LSP = ILocalizationDashboardModule::Get().GetCurrentLocalizationServiceProvider();
+				//TSharedPtr<FLocalizationServiceProviderWrapper>* LSPWrapper = Providers.FindByPredicate(FindProviderPredicate(LSP));
+				//ServiceProviderDetailWidgetRow.ValueContent()
+				//	[
+				//		SNew(SComboBox<TSharedPtr<FLocalizationServiceProviderWrapper>>)
+				//		.InitiallySelectedItem(LSPWrapper ? *LSPWrapper : nullptr)
+				//		.OptionsSource(&Providers)
+				//		.OnGenerateWidget(this, &FProjectLocalizationSettingsDetailsCustomization::ServiceProviderComboBox_OnGenerateWidget)
+				//		.OnSelectionChanged(this, &FProjectLocalizationSettingsDetailsCustomization::ServiceProviderComboBox_OnSelectionChanged)
+				//		.Content()
+				//		[
+				//			SNew(STextBlock)
+				//			.Text(this, &FProjectLocalizationSettingsDetailsCustomization::GetCurrentServiceProviderDisplayName)
+				//		]
+				//	];
+
+				//if (LSP)
+				//{
+				//	LSP->CustomizeSettingsDetails(*ServiceProviderCategoryBuilder);
+				//}
+			}
+		}
+
+		{
+			IDetailCategoryBuilder& DetailCategoryBuilder = DetailBuilder.EditCategory(FName("Targets"));
+			TargetObjectsPropertyHandle = DetailBuilder.GetProperty(GET_MEMBER_NAME_CHECKED(UProjectLocalizationSettings,TargetObjects));
+			if (TargetObjectsPropertyHandle->IsValidHandle())
+			{
+				TargetObjectsPropertyHandle->MarkHiddenByCustomization();
+				TargetsArrayPropertyHandle_OnNumElementsChanged = FSimpleDelegate::CreateSP(this, &FProjectLocalizationSettingsDetailsCustomization::RebuildTargetsList);
+				TargetObjectsPropertyHandle->AsArray()->SetOnNumElementsChanged(TargetsArrayPropertyHandle_OnNumElementsChanged);
+
+				FLocalizationDashboardCommands::Register();
+				const TSharedRef< FUICommandList > CommandList = MakeShareable(new FUICommandList);
+				FToolBarBuilder ToolBarBuilder( CommandList, FMultiBoxCustomization::AllowCustomization("LocalizationDashboard") );
+
+				CommandList->MapAction( FLocalizationDashboardCommands::Get().GatherAllTargets, FExecuteAction::CreateSP(this, &FProjectLocalizationSettingsDetailsCustomization::GatherAllTargets));
+				ToolBarBuilder.AddToolBarButton(FLocalizationDashboardCommands::Get().GatherAllTargets, NAME_None, TAttribute<FText>(), TAttribute<FText>(), FSlateIcon(FEditorStyle::GetStyleSetName(), "LocalizationDashboard.GatherAllTargets"));
+
+				CommandList->MapAction( FLocalizationDashboardCommands::Get().ImportAllTargets, FExecuteAction::CreateSP(this, &FProjectLocalizationSettingsDetailsCustomization::ImportAllTargets));
+				ToolBarBuilder.AddToolBarButton(FLocalizationDashboardCommands::Get().ImportAllTargets, NAME_None, TAttribute<FText>(), TAttribute<FText>(), FSlateIcon(FEditorStyle::GetStyleSetName(), "LocalizationDashboard.ImportForAllTargetsCultures"));
+
+				CommandList->MapAction( FLocalizationDashboardCommands::Get().ExportAllTargets, FExecuteAction::CreateSP(this, &FProjectLocalizationSettingsDetailsCustomization::ExportAllTargets));
+				ToolBarBuilder.AddToolBarButton(FLocalizationDashboardCommands::Get().ExportAllTargets, NAME_None, TAttribute<FText>(), TAttribute<FText>(), FSlateIcon(FEditorStyle::GetStyleSetName(), "LocalizationDashboard.ExportForAllTargetsCultures"));
+
+				BuildTargetsList();
+
+				DetailCategoryBuilder.AddCustomRow( LOCTEXT("TargetsFilterString", "Targets") )
+					.WholeRowContent()
+					[
+						SNew(SVerticalBox)
+						+SVerticalBox::Slot()
+						.AutoHeight()
+						[
+							ToolBarBuilder.MakeWidget()
+						]
+						+SVerticalBox::Slot()
+							.AutoHeight()
+							[
+								SAssignNew(TargetsListView, SListView< TSharedPtr<IPropertyHandle> >)
+								.OnGenerateRow(this, &FProjectLocalizationSettingsDetailsCustomization::OnGenerateRow)
+								.ListItemsSource(&TargetsList)
+								.SelectionMode(ESelectionMode::None)
+								.HeaderRow
+								(
+								SNew(SHeaderRow)
+								+SHeaderRow::Column("Target")
+								.DefaultLabel( LOCTEXT("TargetColumnLabel", "Target"))
+								.HAlignHeader(HAlign_Left)
+								.HAlignCell(HAlign_Left)
+								.VAlignCell(VAlign_Center)
+								+SHeaderRow::Column("Status")
+								.DefaultLabel( LOCTEXT("StatusColumnLabel", "Status"))
+								.HAlignHeader(HAlign_Center)
+								.HAlignCell(HAlign_Center)
+								.VAlignCell(VAlign_Center)
+								+SHeaderRow::Column("Cultures")
+								.DefaultLabel( LOCTEXT("CulturesColumnLabel", "Cultures"))
+								.HAlignHeader(HAlign_Fill)
+								.HAlignCell(HAlign_Fill)
+								.VAlignCell(VAlign_Center)
+								+SHeaderRow::Column("WordCount")
+								.DefaultLabel( LOCTEXT("WordCountColumnLabel", "Word Count"))
+								.HAlignHeader(HAlign_Center)
+								.HAlignCell(HAlign_Center)
+								.VAlignCell(VAlign_Center)
+								+SHeaderRow::Column("Actions")
+								.DefaultLabel( LOCTEXT("ActionsColumnLabel", "Actions"))
+								.HAlignHeader(HAlign_Center)
+								.HAlignCell(HAlign_Center)
+								.VAlignCell(VAlign_Center)
+								)
+							]
+						+SVerticalBox::Slot()
+							.AutoHeight()
+							.VAlign(VAlign_Center)
+							[
+								SNew(SButton)
+								.Text(LOCTEXT("AddNewTargetButtonLabel", "Add New Target"))
+								.OnClicked(this, &FProjectLocalizationSettingsDetailsCustomization::OnNewTargetButtonClicked)
+							]
+					];
+			}
+		}
+	}
+
+	void FProjectLocalizationSettingsDetailsCustomization::BuildTargetsList()
+	{
+		const TSharedPtr<IPropertyHandleArray> TargetObjectsArrayPropertyHandle = TargetObjectsPropertyHandle->AsArray();
+		if (TargetObjectsArrayPropertyHandle.IsValid())
+		{
+			uint32 ElementCount = 0;
+			TargetObjectsArrayPropertyHandle->GetNumElements(ElementCount);
+			for (uint32 i = 0; i < ElementCount; ++i)
+			{
+				const TSharedPtr<IPropertyHandle> TargetObjectPropertyHandle = TargetObjectsArrayPropertyHandle->GetElement(i);
+				TargetsList.Add(TargetObjectPropertyHandle);
+			}
+		}
+	}
+
+	void FProjectLocalizationSettingsDetailsCustomization::RebuildTargetsList()
+	{
+		const TSharedPtr<IPropertyHandleArray> TargetObjectsArrayPropertyHandle = TargetObjectsPropertyHandle->AsArray();
+		if (TargetObjectsArrayPropertyHandle.IsValid() && NewEntryIndexToBeInitialized != INDEX_NONE)
+		{
+			const TSharedPtr<IPropertyHandle> TargetObjectPropertyHandle = TargetObjectsArrayPropertyHandle->GetElement(NewEntryIndexToBeInitialized);
+			if (TargetObjectPropertyHandle.IsValid() && TargetObjectPropertyHandle->IsValidHandle())
+			{
+				ULocalizationTarget* const NewTarget = Cast<ULocalizationTarget>(StaticConstructObject(ULocalizationTarget::StaticClass(), Settings));
+
+				TArray<void*> RawData;
+				TargetObjectsPropertyHandle->AccessRawData(RawData);
+				void* RawDatum = RawData.Top();
+				TArray<ULocalizationTarget*>& TargetObjectsArray = *(reinterpret_cast< TArray<ULocalizationTarget*>* >(RawDatum));
+
+				FString NewTargetName = "NewTarget";
+				auto TargetNameIsIdentical = [&](ULocalizationTarget* Target) -> bool
+				{
+					return (Target != NewTarget) && Target && (Target->Settings.Name == NewTargetName);
+				};
+
+				for (uint32 i = 0; TargetObjectsArray.ContainsByPredicate(TargetNameIsIdentical); ++i)
+				{
+					NewTargetName = FString::Printf(TEXT("NewTarget%u"), i);
+				}
+
+				NewTarget->Settings.Name = NewTargetName;
+				NewTarget->Settings.NativeCultureStatistics.CultureName = FInternationalization::Get().GetCurrentCulture()->GetName();
+
+				const UObject* SetValue = NewTarget;
+				TargetObjectPropertyHandle->SetValue(SetValue);
+
+				NewEntryIndexToBeInitialized = INDEX_NONE;
+			}
+		}
+
+		TargetsList.Empty();
+		BuildTargetsList();
+
+		if (TargetsListView.IsValid())
+		{
+			TargetsListView->RequestListRefresh();
+		}
+	}
+
+	FText FProjectLocalizationSettingsDetailsCustomization::GetCurrentServiceProviderDisplayName() const
+	{
+		ILocalizationServiceProvider* const LSP = ILocalizationDashboardModule::Get().GetCurrentLocalizationServiceProvider();
+		return LSP ? LSP->GetDisplayName() : LOCTEXT("NoServiceProviderName", "None");
+	}
+
+	TSharedRef<SWidget> FProjectLocalizationSettingsDetailsCustomization::ServiceProviderComboBox_OnGenerateWidget(TSharedPtr<FLocalizationServiceProviderWrapper> LSPWrapper) const
+	{
+		ILocalizationServiceProvider* const LSP = LSPWrapper->Provider;
+		return	SNew(STextBlock)
+			.Text(LSP ? LSP->GetDisplayName() : LOCTEXT("NoServiceProviderName", "None"));
+	}
+
+	void FProjectLocalizationSettingsDetailsCustomization::ServiceProviderComboBox_OnSelectionChanged(TSharedPtr<FLocalizationServiceProviderWrapper> LSPWrapper, ESelectInfo::Type SelectInfo)
+	{
+		ILocalizationServiceProvider* const LSP = LSPWrapper->Provider;
+		FString ServiceProviderName = LSP ? LSP->GetName().ToString() : TEXT("None");
+		ServiceProviderPropertyHandle->SetValue(ServiceProviderName);
+
+		if (LSP)
+		{
+			LSP->CustomizeSettingsDetails(*ServiceProviderCategoryBuilder);
+		}
+		DetailLayoutBuilder->ForceRefreshDetails();
+	}
+
+	void FProjectLocalizationSettingsDetailsCustomization::GatherAllTargets()
+	{
+		for (ULocalizationTarget* const LocalizationTarget : Settings->TargetObjects)
+		{
+			// Save unsaved packages.
+			bool DidPackagesNeedSaving;
+			const bool WerePackagesSaved = FEditorFileUtils::SaveDirtyPackages(true, true, true, false, false, &DidPackagesNeedSaving);
+
+			if (DidPackagesNeedSaving && !WerePackagesSaved)
+			{
+				// Give warning dialog.
+				const FText MessageText = NSLOCTEXT("LocalizationCultureActions", "UnsavedPackagesWarningDialogMessage", "There are unsaved changes. These changes may not be gathered from correctly.");
+				const FText TitleText = NSLOCTEXT("LocalizationCultureActions", "UnsavedPackagesWarningDialogTitle", "Unsaved Changes Before Gather");
+				switch(FMessageDialog::Open(EAppMsgType::OkCancel, MessageText, &TitleText))
+				{
+				case EAppReturnType::Cancel:
+					{
+						return;
+					}
+					break;
+				}
+			}
+
+		}
+
+		// Execute gather.
+		const TSharedPtr<SWindow> ParentWindow = FSlateApplication::Get().FindWidgetWindow(DetailLayoutBuilder->GetDetailsView().AsShared());
+		TArray<FLocalizationTargetSettings*> TargetSettings;
+		for (ULocalizationTarget* const TargetObject : Settings->TargetObjects)
+		{
+			TargetSettings.Add(&TargetObject->Settings);
+		}
+		LocalizationCommandletTasks::GatherTargets(ParentWindow.ToSharedRef(), TargetSettings);
+
+		for (ULocalizationTarget* const LocalizationTarget : Settings->TargetObjects)
+		{
+			UpdateTargetFromReports(LocalizationTarget);
+		}
+	}
+
+	void FProjectLocalizationSettingsDetailsCustomization::ImportAllTargets()
+	{
+		const TSharedPtr<SWindow> ParentWindow = FSlateApplication::Get().FindWidgetWindow(DetailLayoutBuilder->GetDetailsView().AsShared());
+		TArray<FLocalizationTargetSettings*> TargetSettings;
+		for (ULocalizationTarget* const TargetObject : Settings->TargetObjects)
+		{
+			TargetSettings.Add(&TargetObject->Settings);
+		}
+		LocalizationCommandletTasks::ImportTargets(ParentWindow.ToSharedRef(), TargetSettings);
+
+		for (ULocalizationTarget* const LocalizationTarget : Settings->TargetObjects)
+		{
+			UpdateTargetFromReports(LocalizationTarget);
+		}
+	}
+
+	void FProjectLocalizationSettingsDetailsCustomization::ExportAllTargets()
+	{
+		const TSharedPtr<SWindow> ParentWindow = FSlateApplication::Get().FindWidgetWindow(DetailLayoutBuilder->GetDetailsView().AsShared());
+		TArray<FLocalizationTargetSettings*> TargetSettings;
+		for (ULocalizationTarget* const TargetObject : Settings->TargetObjects)
+		{
+			TargetSettings.Add(&TargetObject->Settings);
+		}
+		LocalizationCommandletTasks::ExportTargets(ParentWindow.ToSharedRef(), TargetSettings);
+	}
+
+	void FProjectLocalizationSettingsDetailsCustomization::UpdateTargetFromReports(ULocalizationTarget* const LocalizationTarget)
+	{
+		//TArray< TSharedPtr<IPropertyHandle> > WordCountPropertyHandles;
+
+		//const TSharedPtr<IPropertyHandle> TargetSettingsPropertyHandle = TargetEditor->GetTargetSettingsPropertyHandle();
+		//if (TargetSettingsPropertyHandle.IsValid() && TargetSettingsPropertyHandle->IsValidHandle())
+		//{
+		//	const TSharedPtr<IPropertyHandle> NativeCultureStatisticsPropertyHandle = TargetSettingsPropertyHandle->GetChildHandle(GET_MEMBER_NAME_CHECKED(FLocalizationTargetSettings, NativeCultureStatistics));
+		//	if (NativeCultureStatisticsPropertyHandle.IsValid() && NativeCultureStatisticsPropertyHandle->IsValidHandle())
+		//	{
+		//		const TSharedPtr<IPropertyHandle> WordCountPropertyHandle = NativeCultureStatisticsPropertyHandle->GetChildHandle(GET_MEMBER_NAME_CHECKED(FCultureStatistics, WordCount));
+		//		if (WordCountPropertyHandle.IsValid() && WordCountPropertyHandle->IsValidHandle())
+		//		{
+		//			WordCountPropertyHandles.Add(WordCountPropertyHandle);
+		//		}
+		//	}
+
+		//	const TSharedPtr<IPropertyHandle> SupportedCulturesStatisticsPropertyHandle = TargetSettingsPropertyHandle->GetChildHandle(GET_MEMBER_NAME_CHECKED(FLocalizationTargetSettings, SupportedCulturesStatistics));
+		//	if (SupportedCulturesStatisticsPropertyHandle.IsValid() && SupportedCulturesStatisticsPropertyHandle->IsValidHandle())
+		//	{
+		//		uint32 SupportedCultureCount = 0;
+		//		SupportedCulturesStatisticsPropertyHandle->GetNumChildren(SupportedCultureCount);
+		//		for (uint32 i = 0; i < SupportedCultureCount; ++i)
+		//		{
+		//			const TSharedPtr<IPropertyHandle> ElementPropertyHandle = SupportedCulturesStatisticsPropertyHandle->GetChildHandle(i);
+		//			if (ElementPropertyHandle.IsValid() && ElementPropertyHandle->IsValidHandle())
+		//			{
+		//				const TSharedPtr<IPropertyHandle> WordCountPropertyHandle = SupportedCulturesStatisticsPropertyHandle->GetChildHandle(GET_MEMBER_NAME_CHECKED(FCultureStatistics, WordCount));
+		//				if (WordCountPropertyHandle.IsValid() && WordCountPropertyHandle->IsValidHandle())
+		//				{
+		//					WordCountPropertyHandles.Add(WordCountPropertyHandle);
+		//				}
+		//			}
+		//		}
+		//	}
+		//}
+
+		//for (const TSharedPtr<IPropertyHandle>& WordCountPropertyHandle : WordCountPropertyHandles)
+		//{
+		//	WordCountPropertyHandle->NotifyPreChange();
+		//}
+		LocalizationTarget->Settings.UpdateWordCountsFromCSV();
+		LocalizationTarget->Settings.UpdateStatusFromConflictReport();
+		//for (const TSharedPtr<IPropertyHandle>& WordCountPropertyHandle : WordCountPropertyHandles)
+		//{
+		//	WordCountPropertyHandle->NotifyPostChange();
+		//}
+	}
+
+	TSharedRef<ITableRow> FProjectLocalizationSettingsDetailsCustomization::OnGenerateRow(TSharedPtr<IPropertyHandle> TargetObjectPropertyHandle, const TSharedRef<STableViewBase>& Table)
+	{
+		return SNew(SLocalizationDashboardTargetRow, Table, DetailLayoutBuilder->GetPropertyUtilities(), TargetObjectPropertyHandle.ToSharedRef());
+	}
+
+	FReply FProjectLocalizationSettingsDetailsCustomization::OnNewTargetButtonClicked()
+	{
+		if(TargetObjectsPropertyHandle.IsValid() && TargetObjectsPropertyHandle->IsValidHandle())
+		{
+			uint32 NewElementIndex;
+			TargetObjectsPropertyHandle->AsArray()->GetNumElements(NewElementIndex);
+			TargetObjectsPropertyHandle->AsArray()->AddItem();
+			NewEntryIndexToBeInitialized = NewElementIndex;
+		}
+
+		return FReply::Handled();
+	}
+}
+
+class SLocalizationDashboard : public SCompoundWidget
+{
+public:
+	SLATE_BEGIN_ARGS(SLocalizationDashboard) {}
+	SLATE_END_ARGS()
+
+		void Construct(const FArguments& InArgs, const TSharedPtr<SWindow>& OwningWindow, const TSharedRef<SDockTab>& OwningTab, UProjectLocalizationSettings* const Settings);
+	TWeakPtr<SDockTab> ShowTargetEditor(ULocalizationTarget* const LocalizationTarget);
+
+private:
+	static const FName DetailsTabName;
+	static const FName DocumentsTabName;
+
+	TSharedPtr<FTabManager> TabManager;
+	TMap< TWeakObjectPtr<ULocalizationTarget>, TWeakPtr<SDockTab> > TargetToTabMap;
+};
+
+const FName SLocalizationDashboard::DetailsTabName("Details");
+const FName SLocalizationDashboard::DocumentsTabName("Documents");
+
+void SLocalizationDashboard::Construct(const FArguments& InArgs, const TSharedPtr<SWindow>& OwningWindow, const TSharedRef<SDockTab>& OwningTab, UProjectLocalizationSettings* const Settings)
+{
+	check(Settings);
+
+	TabManager = FGlobalTabmanager::Get()->NewTabManager(OwningTab);
+
+	const auto& PersistLayout = [](const TSharedRef<FTabManager::FLayout>& LayoutToSave)
+	{
+		FLayoutSaveRestore::SaveToConfig(GEditorLayoutIni, LayoutToSave);
+	};
+	TabManager->SetOnPersistLayout(FTabManager::FOnPersistLayout::CreateLambda(PersistLayout));
+
+	const auto& CreateDetailsTab = [Settings](const FSpawnTabArgs& SpawnTabArgs) -> TSharedRef<SDockTab>
+	{
+		const TSharedRef<SDockTab> DockTab = SNew(SDockTab);
+			//.OnCanCloseTab_Lambda([]()->bool{return false;});
+
+		FPropertyEditorModule& PropertyModule = FModuleManager::LoadModuleChecked<FPropertyEditorModule>("PropertyEditor");
+		FDetailsViewArgs DetailsViewArgs(false, false, false, FDetailsViewArgs::ENameAreaSettings::HideNameArea, false, nullptr, false, NAME_None);
+		TSharedRef<IDetailsView> DetailsView = PropertyModule.CreateDetailView(DetailsViewArgs);
+
+		DetailsView->RegisterInstancedCustomPropertyLayout(UProjectLocalizationSettings::StaticClass(), FOnGetDetailCustomizationInstance::CreateLambda([&]()->TSharedRef<IDetailCustomization>{return MakeShareable(new FProjectLocalizationSettingsDetailsCustomization(Settings));}));
+		DetailsView->SetObject(Settings, true);
+
+		DockTab->SetContent(DetailsView);
+
+		return DockTab;
+	};
+	TabManager->RegisterTabSpawner(DetailsTabName, FOnSpawnTab::CreateLambda(CreateDetailsTab));
+
+	TSharedRef<FTabManager::FLayout> Layout = FTabManager::NewLayout("LocalizationDashboard_Experimental_V1")
+		->AddArea
+		(
+		FTabManager::NewPrimaryArea()
+		->SetOrientation(Orient_Vertical)
+		->Split
+		(
+		FTabManager::NewStack()
+		->SetHideTabWell(true)
+		->AddTab(DetailsTabName, ETabState::OpenedTab)
+		)
+		->Split
+		(
+		FTabManager::NewStack()
+		->SetHideTabWell(false)
+		->AddTab(DocumentsTabName, ETabState::ClosedTab)
+		)
+		);
+
+	Layout = FLayoutSaveRestore::LoadFromConfig(GEditorLayoutIni, Layout);
+
+	ChildSlot
+		[
+			TabManager->RestoreFrom(Layout, OwningWindow).ToSharedRef()
+		];
+}
+
+TWeakPtr<SDockTab> SLocalizationDashboard::ShowTargetEditor(ULocalizationTarget* const LocalizationTarget)
+{
+	// Create tab if not existent.
+	TWeakPtr<SDockTab>& TargetEditorDockTab = TargetToTabMap.FindOrAdd(TWeakObjectPtr<ULocalizationTarget>(LocalizationTarget));
+
+	if (!TargetEditorDockTab.IsValid())
+	{
+		UProjectLocalizationSettings* const ProjectSettings = LocalizationTarget->GetTypedOuter<UProjectLocalizationSettings>();
+
+		const TSharedRef<SLocalizationTargetEditor> OurTargetEditor = SNew(SLocalizationTargetEditor, ProjectSettings, LocalizationTarget);
+		const TSharedRef<SDockTab> NewTargetEditorTab = SNew(SDockTab)
+			.TabRole(ETabRole::DocumentTab)
+			.Label_Lambda( [LocalizationTarget]
+			{
+				return LocalizationTarget ? FText::FromString(LocalizationTarget->Settings.Name) : FText::GetEmpty();
+			})
+			[
+				OurTargetEditor
+			];
+
+		TabManager->InsertNewDocumentTab(DocumentsTabName, FTabManager::ESearchPreference::RequireClosedTab, NewTargetEditorTab);
+		TargetEditorDockTab = NewTargetEditorTab;
+	}
+	else
+	{
+		const TSharedPtr<SDockTab> OldTargetEditorTab = TargetEditorDockTab.Pin();
+		TabManager->DrawAttention(OldTargetEditorTab.ToSharedRef());
+	}
+
+	return TargetEditorDockTab;
+}
+
+FLocalizationDashboard* FLocalizationDashboard::Instance = nullptr;
+
+FLocalizationDashboard* FLocalizationDashboard::Get()
+{
+	return Instance;
+}
+
+void FLocalizationDashboard::Initialize()
+{
+	if (!Instance)
+	{
+		Instance = new FLocalizationDashboard();
+	}
+}
+
+void FLocalizationDashboard::Terminate()
+{
+	if (Instance)
+	{
+		delete Instance;
+	}
+}
+
+void FLocalizationDashboard::Show()
+{
+	FGlobalTabmanager::Get()->InvokeTab(TabName);
+}
+
+TWeakPtr<SDockTab> FLocalizationDashboard::ShowTargetEditorTab(ULocalizationTarget* const LocalizationTarget)
+{
+	if (LocalizationDashboardWidget.IsValid())
+	{
+		return LocalizationDashboardWidget->ShowTargetEditor(LocalizationTarget);
+	}
+	return nullptr;
+}
+
+FLocalizationDashboard::FLocalizationDashboard()
+{
+	RegisterTabSpawner();
+}
+
+FLocalizationDashboard::~FLocalizationDashboard()
+{
+	UnregisterTabSpawner();
+}
+
+void FLocalizationDashboard::RegisterTabSpawner()
+{
+	const auto& SpawnMainTab = [this](const FSpawnTabArgs& Args) -> TSharedRef<SDockTab>
+	{
+		const TSharedRef<SDockTab> DockTab = SNew(SDockTab)
+			.Label(LOCTEXT("MainTabTitle", "Localization Dashboard"))
+			.TabRole(ETabRole::MajorTab);
+
+		DockTab->SetContent( SAssignNew(LocalizationDashboardWidget, SLocalizationDashboard, Args.GetOwnerWindow(), DockTab, GetMutableDefault<UProjectLocalizationSettings>()) );
+
+		return DockTab;
+	};
+
+	FGlobalTabmanager::Get()->RegisterTabSpawner(TabName, FOnSpawnTab::CreateLambda( TFunction<TSharedRef<SDockTab> (const FSpawnTabArgs&)>(SpawnMainTab) ) )
+		.SetDisplayName(LOCTEXT("MainTabTitle", "Localization Dashboard"))
+		.SetMenuType(ETabSpawnerMenuType::Hidden);
+}
+
+void FLocalizationDashboard::UnregisterTabSpawner()
+{
+	FGlobalTabmanager::Get()->UnregisterTabSpawner(TabName);
+}
+
+#undef LOCTEXT_NAMESPACE
\ No newline at end of file
diff --git a/Engine/Source/Editor/LocalizationDashboard/Private/LocalizationDashboard.h b/Engine/Source/Editor/LocalizationDashboard/Private/LocalizationDashboard.h
new file mode 100644
--- /dev/null
+++ b/Engine/Source/Editor/LocalizationDashboard/Private/LocalizationDashboard.h
@@ -0,0 +1,29 @@
+// Copyright 1998-2014 Epic Games, Inc. All Rights Reserved.
+
+#pragma once
+
+#include "SCompoundWidget.h"
+
+class UProjectLocalizationSettings;
+class SLocalizationDashboard;
+
+class FLocalizationDashboard
+{
+public:
+	static FLocalizationDashboard* Get();
+	static void Initialize();
+	static void Terminate();
+	void Show();
+	TWeakPtr<SDockTab> ShowTargetEditorTab(ULocalizationTarget* const LocalizationTarget);
+
+private:
+	FLocalizationDashboard();
+	~FLocalizationDashboard();
+	void RegisterTabSpawner();
+	void UnregisterTabSpawner();
+	
+private:
+	static const FName TabName;
+	static FLocalizationDashboard* Instance;
+	TSharedPtr<SLocalizationDashboard> LocalizationDashboardWidget;
+};
\ No newline at end of file
diff --git a/Engine/Source/Editor/LocalizationDashboard/Private/LocalizationDashboardModule.cpp b/Engine/Source/Editor/LocalizationDashboard/Private/LocalizationDashboardModule.cpp
--- a/Engine/Source/Editor/LocalizationDashboard/Private/LocalizationDashboardModule.cpp
+++ b/Engine/Source/Editor/LocalizationDashboard/Private/LocalizationDashboardModule.cpp
@@ -2,7 +2,7 @@
 
 #include "LocalizationDashboardPrivatePCH.h"
 #include "ILocalizationDashboardModule.h"
-#include "SLocalizationDashboard.h"
+#include "LocalizationDashboard.h"
 #include "ProjectLocalizationSettings.h"
 #include "SDockTab.h"
 #include "Features/IModularFeatures.h"
@@ -28,19 +28,7 @@ public:
 	// Begin IModuleInterface interface
 	virtual void StartupModule() override
 	{
-		const auto& SpawnMainTab = [this](const FSpawnTabArgs& Args) -> TSharedRef<SDockTab>
-		{
-			return SNew(SDockTab)
-				.Label(LOCTEXT("MainTabTitle", "Localization Dashboard"))
-				.TabRole(ETabRole::MajorTab)
-				[
-					SNew(SLocalizationDashboard, Settings.Get())
-				];
-		};
-
-		FGlobalTabmanager::Get()->RegisterTabSpawner(SLocalizationDashboard::TabName, FOnSpawnTab::CreateLambda( TFunction<TSharedRef<SDockTab> (const FSpawnTabArgs&)>(SpawnMainTab) ) )
-			.SetDisplayName(LOCTEXT("MainTabTitle", "Localization Dashboard"))
-			.SetMenuType(ETabSpawnerMenuType::Hidden);
+		FLocalizationDashboard::Initialize();
 
 		ServiceProviders = IModularFeatures::Get().GetModularFeatureImplementations<ILocalizationServiceProvider>("LocalizationService");
 		ServiceProviders.Insert(nullptr, 0); // "None"
@@ -48,13 +36,17 @@ public:
 
 	virtual void ShutdownModule() override
 	{
-		FGlobalTabmanager::Get()->UnregisterTabSpawner(SLocalizationDashboard::TabName);
+		FLocalizationDashboard::Terminate();
 	}
 	// End IModuleInterface interface
 
 	virtual void Show()
 	{
-		FGlobalTabmanager::Get()->InvokeTab(SLocalizationDashboard::TabName);
+		FLocalizationDashboard* const LocalizationDashboard = FLocalizationDashboard::Get();
+		if (LocalizationDashboard)
+		{
+			LocalizationDashboard->Show();
+		}
 	}
 
 	virtual const TArray<ILocalizationServiceProvider*>& GetLocalizationServiceProviders() const
diff --git a/Engine/Source/Editor/LocalizationDashboard/Private/SLocalizationDashboard.cpp b/Engine/Source/Editor/LocalizationDashboard/Private/SLocalizationDashboard.cpp
deleted file mode 100644
diff --git a/Engine/Source/Editor/LocalizationDashboard/Private/SLocalizationDashboard.h b/Engine/Source/Editor/LocalizationDashboard/Private/SLocalizationDashboard.h
deleted file mode 100644
diff --git a/Engine/Source/Editor/LocalizationDashboard/Private/SLocalizationDashboardTargetRow.cpp b/Engine/Source/Editor/LocalizationDashboard/Private/SLocalizationDashboardTargetRow.cpp
--- a/Engine/Source/Editor/LocalizationDashboard/Private/SLocalizationDashboardTargetRow.cpp
+++ b/Engine/Source/Editor/LocalizationDashboard/Private/SLocalizationDashboardTargetRow.cpp
@@ -1,7 +1,7 @@
 // Copyright 1998-2014 Epic Games, Inc. All Rights Reserved.
 
 #include "LocalizationDashboardPrivatePCH.h"
-#include "SLocalizationDashboard.h"
+#include "LocalizationDashboard.h"
 #include "SLocalizationDashboardTargetRow.h"
 #include "SHyperlink.h"
 #include "SLocalizationTargetEditor.h"
@@ -11,8 +11,6 @@
 
 #define LOCTEXT_NAMESPACE "LocalizationDashboardTargetRow"
 
-TMap< TWeakObjectPtr<ULocalizationTarget>, TWeakPtr<SDockTab> > SLocalizationDashboardTargetRow::TargetToTabMap;
-
 void SLocalizationDashboardTargetRow::Construct(const FTableRowArgs& InArgs, const TSharedRef<STableViewBase>& OwnerTableView, const TSharedRef<IPropertyUtilities>& InPropertyUtilities, const TSharedRef<IPropertyHandle>& InTargetObjectPropertyHandle)
 {
 	PropertyUtilities = InPropertyUtilities;
@@ -177,38 +175,11 @@ void SLocalizationDashboardTargetRow::OnNavigate()
 	ULocalizationTarget* const LocalizationTarget = GetTarget();
 	if (LocalizationTarget)
 	{
-		// Create tab if not existent.
-		TWeakPtr<SDockTab>& TargetEditorDockTab = TargetToTabMap.FindOrAdd(TWeakObjectPtr<ULocalizationTarget>(LocalizationTarget));
-
-		if (!TargetEditorDockTab.IsValid())
-		{
-			UProjectLocalizationSettings* ProjectSettings = nullptr;
-			const TSharedPtr<IPropertyHandle> TargetObjectsPropertyHandle = TargetObjectPropertyHandle->GetParentHandle(); // TargetObjects
-			if (TargetObjectsPropertyHandle.IsValid() && TargetObjectsPropertyHandle->IsValidHandle())
-			{
-				TArray<UObject*> OuterObjects;
-				TargetObjectsPropertyHandle->GetOuterObjects(OuterObjects);
-				ProjectSettings = Cast<UProjectLocalizationSettings>(OuterObjects.Top());
-			}
+		FLocalizationDashboard* const LocalizationDashboard = FLocalizationDashboard::Get();
 
-			const TSharedRef<SLocalizationTargetEditor> OurTargetEditor = SNew(SLocalizationTargetEditor, ProjectSettings, LocalizationTarget);
-			const TSharedRef<SDockTab> NewTargetEditorTab = SNew(SDockTab)
-				.TabRole(ETabRole::DocumentTab)
-				.Label_Lambda( [LocalizationTarget]
-			{
-				return LocalizationTarget ? FText::FromString(LocalizationTarget->Settings.Name) : FText::GetEmpty();
-			})
-				[
-					OurTargetEditor
-				];
-
-			FGlobalTabmanager::Get()->InsertNewDocumentTab( SLocalizationDashboard::TabName, FTabManager::ESearchPreference::RequireClosedTab, NewTargetEditorTab );
-			TargetEditorDockTab = NewTargetEditorTab;
-		}
-		else
+		if (LocalizationDashboard)
 		{
-			const TSharedPtr<SDockTab> OldTargetEditorTab = TargetEditorDockTab.Pin();
-			FGlobalTabmanager::Get()->DrawAttention(OldTargetEditorTab.ToSharedRef());
+			TargetEditorDockTab = LocalizationDashboard->ShowTargetEditorTab(LocalizationTarget);
 		}
 	}
 }
@@ -427,7 +398,6 @@ void SLocalizationDashboardTargetRow::Delete()
 					LocalizationTarget->Settings.DeleteFiles();
 
 					// Close target editor.
-					TWeakPtr<SDockTab>& TargetEditorDockTab = TargetToTabMap.FindOrAdd(TWeakObjectPtr<ULocalizationTarget>(LocalizationTarget));
 					if (TargetEditorDockTab.IsValid())
 					{
 						const TSharedPtr<SDockTab> OldTargetEditorTab = TargetEditorDockTab.Pin();
diff --git a/Engine/Source/Editor/LocalizationDashboard/Private/SLocalizationDashboardTargetRow.h b/Engine/Source/Editor/LocalizationDashboard/Private/SLocalizationDashboardTargetRow.h
--- a/Engine/Source/Editor/LocalizationDashboard/Private/SLocalizationDashboardTargetRow.h
+++ b/Engine/Source/Editor/LocalizationDashboard/Private/SLocalizationDashboardTargetRow.h
@@ -40,5 +40,5 @@ private:
 private:
 	TSharedPtr<IPropertyUtilities> PropertyUtilities;
 	TSharedPtr<IPropertyHandle> TargetObjectPropertyHandle;
-	static TMap< TWeakObjectPtr<ULocalizationTarget>, TWeakPtr<SDockTab> > TargetToTabMap;
+	TWeakPtr<SDockTab> TargetEditorDockTab;
 };
\ No newline at end of file
