Potential fix for an inconsistently reproducible crash when closing the commandlet executor window for localization.

diff --git a/Engine/Source/Editor/LocalizationDashboard/Private/LocalizationCommandletExecution.cpp b/Engine/Source/Editor/LocalizationDashboard/Private/LocalizationCommandletExecution.cpp
--- a/Engine/Source/Editor/LocalizationDashboard/Private/LocalizationCommandletExecution.cpp
+++ b/Engine/Source/Editor/LocalizationDashboard/Private/LocalizationCommandletExecution.cpp
@@ -14,13 +14,10 @@
 
 namespace
 {
-	class SLocalizationCommandletExecutor : public SWindow
+	class SLocalizationCommandletExecutor : public SCompoundWidget
 	{
-		SLATE_BEGIN_ARGS(SLocalizationCommandletExecutor)
-			: _Title( LOCTEXT("WindowTitle", "Localization Commandlet") )
-		{}
-		SLATE_ATTRIBUTE(FText, Title)
-			SLATE_END_ARGS()
+		SLATE_BEGIN_ARGS(SLocalizationCommandletExecutor) {}
+		SLATE_END_ARGS()
 
 	private:
 		struct FTaskListModel
@@ -49,7 +46,7 @@ namespace
 		SLocalizationCommandletExecutor();
 		~SLocalizationCommandletExecutor();
 
-		void Construct(const FArguments& Arguments, const TArray<LocalizationCommandletExecution::FTask>& Tasks);
+		void Construct(const FArguments& Arguments, const TSharedRef<SWindow>& ParentWindow, const TArray<LocalizationCommandletExecution::FTask>& Tasks);
 		void Tick( const FGeometry& AllottedGeometry, const double InCurrentTime, const float InDeltaTime ) override;
 		bool WasSuccessful() const;
 		void Log(const FString& String);
@@ -91,6 +88,7 @@ namespace
 		} PendingData;
 
 
+		TSharedPtr<SWindow> ParentWindow;
 		TSharedPtr<FLocalizationCommandletProcess> CommandletProcess;
 		FRunnable* Runnable;
 		FRunnableThread* RunnableThread;
@@ -108,8 +106,10 @@ namespace
 		CancelCommandlet();
 	}
 
-	void SLocalizationCommandletExecutor::Construct(const FArguments& Arguments, const TArray<LocalizationCommandletExecution::FTask>& Tasks)
+	void SLocalizationCommandletExecutor::Construct(const FArguments& Arguments, const TSharedRef<SWindow>& InParentWindow, const TArray<LocalizationCommandletExecution::FTask>& Tasks)
 	{
+		ParentWindow = InParentWindow;
+
 		for (const LocalizationCommandletExecution::FTask& Task : Tasks)
 		{
 			const TSharedRef<FTaskListModel> Model = MakeShareable(new FTaskListModel());
@@ -120,14 +120,7 @@ namespace
 		TSharedPtr<SScrollBar> VerticalScrollBar;
 		TSharedPtr<SScrollBar> HorizontalScrollBar;
 
-		SWindow::Construct(
-			SWindow::FArguments()
-			.SupportsMinimize(false)
-			.AutoCenter(EAutoCenter::PreferredWorkArea)
-			.ClientSize(FVector2D(600,400))
-			.ActivateWhenFirstShown(true)
-			.FocusWhenFirstShown(true)
-			.Content()
+		ChildSlot
 			[
 				SNew(SVerticalBox)
 				+ SVerticalBox::Slot()
@@ -249,7 +242,7 @@ namespace
 								]
 							]
 					]
-			]);
+			];
 
 		if(TaskListModels.Num() > 0)
 		{
@@ -638,17 +631,26 @@ namespace
 		{
 			CancelCommandlet();
 		}
-		RequestDestroyWindow();
+		ParentWindow->RequestDestroyWindow();
 		return FReply::Handled();
 	}
 }
 
 bool LocalizationCommandletExecution::Execute(const TSharedRef<SWindow>& ParentWindow, const FText& Title, const TArray<FTask>& Tasks)
 {
-	const TSharedRef<SLocalizationCommandletExecutor> CommandletWindow = SNew(SLocalizationCommandletExecutor, Tasks)
-		.Title(Title);
+	const TSharedRef<SWindow> CommandletWindow = SNew(SWindow)
+		.Title(Title)
+		.SupportsMinimize(false)
+		.AutoCenter(EAutoCenter::PreferredWorkArea)
+		.ClientSize(FVector2D(600,400))
+		.ActivateWhenFirstShown(true)
+		.FocusWhenFirstShown(true);
+	const TSharedRef<SLocalizationCommandletExecutor> CommandletExecutor = SNew(SLocalizationCommandletExecutor, CommandletWindow, Tasks);
+	CommandletWindow->SetContent(CommandletExecutor);
+
 	FSlateApplication::Get().AddModalWindow(CommandletWindow, ParentWindow, false);
-	return CommandletWindow->WasSuccessful();
+
+	return CommandletExecutor->WasSuccessful();
 }
 
 TSharedPtr<FLocalizationCommandletProcess> FLocalizationCommandletProcess::Execute(const FString& ConfigFilePath)
