Fixed non-Unity build error in LocalizationDashboard.h.

diff --git a/Engine/Source/Editor/LocalizationDashboard/Private/LocalizationDashboard.h b/Engine/Source/Editor/LocalizationDashboard/Private/LocalizationDashboard.h
--- a/Engine/Source/Editor/LocalizationDashboard/Private/LocalizationDashboard.h
+++ b/Engine/Source/Editor/LocalizationDashboard/Private/LocalizationDashboard.h
@@ -4,7 +4,7 @@
 
 #include "SCompoundWidget.h"
 
-class UProjectLocalizationSettings;
+class ULocalizationTarget;
 class SLocalizationDashboard;
 
 class FLocalizationDashboard
