Added dialog prompts for path selection in import/export operations in the localization dashboard and target editor.

diff --git a/Engine/Source/Editor/LocalizationDashboard/Private/LocalizationCommandletTasks.cpp b/Engine/Source/Editor/LocalizationDashboard/Private/LocalizationCommandletTasks.cpp
--- a/Engine/Source/Editor/LocalizationDashboard/Private/LocalizationCommandletTasks.cpp
+++ b/Engine/Source/Editor/LocalizationDashboard/Private/LocalizationCommandletTasks.cpp
@@ -50,7 +50,7 @@ bool LocalizationCommandletTasks::GatherTarget(const TSharedRef<SWindow>& Parent
 }
 
 
-bool LocalizationCommandletTasks::ImportTargets(const TSharedRef<SWindow>& ParentWindow, const TArray<FLocalizationTargetSettings*>& TargetsSettings)
+bool LocalizationCommandletTasks::ImportTargets(const TSharedRef<SWindow>& ParentWindow, const TArray<FLocalizationTargetSettings*>& TargetsSettings, const TOptional<FString> DirectoryPath)
 {
 	TArray<LocalizationCommandletExecution::FTask> Tasks;
 
@@ -60,8 +60,9 @@ bool LocalizationCommandletTasks::ImportTargets(const TSharedRef<SWindow>& Paren
 		Arguments.Add(TEXT("TargetName"), FText::FromString(TargetSettings->Name));
 
 		const FText ImportTaskName = FText::Format(LOCTEXT("ImportTaskNameFormat", "Import Translations for {TargetName}"), Arguments);
-		const FString ImportScriptPath = LocalizationConfigurationScript::GetImportScriptPath(*TargetSettings);
-		LocalizationConfigurationScript::GenerateImportScript(*TargetSettings).Write(ImportScriptPath);
+		const FString ImportScriptPath = LocalizationConfigurationScript::GetImportScriptPath(*TargetSettings, TOptional<FString>());
+		const TOptional<FString> DirectoryPathForTarget = DirectoryPath.IsSet() ? DirectoryPath.GetValue() / TargetSettings->Name : TOptional<FString>();
+		LocalizationConfigurationScript::GenerateImportScript(*TargetSettings, TOptional<FString>(), DirectoryPathForTarget).Write(ImportScriptPath);
 		Tasks.Add(LocalizationCommandletExecution::FTask(ImportTaskName, ImportScriptPath));
 
 		const FText ReportTaskName = FText::Format(LOCTEXT("ReportTaskNameFormat", "Generate Reports for {TargetName}"), Arguments);
@@ -73,12 +74,12 @@ bool LocalizationCommandletTasks::ImportTargets(const TSharedRef<SWindow>& Paren
 	return LocalizationCommandletExecution::Execute(ParentWindow, LOCTEXT("ImportForAllTargetsWindowTitle", "Import Translations for All Targets"), Tasks);
 }
 
-bool LocalizationCommandletTasks::ImportTarget(const TSharedRef<SWindow>& ParentWindow, FLocalizationTargetSettings& TargetSettings)
+bool LocalizationCommandletTasks::ImportTarget(const TSharedRef<SWindow>& ParentWindow, FLocalizationTargetSettings& TargetSettings, const TOptional<FString> DirectoryPath)
 {
 	TArray<LocalizationCommandletExecution::FTask> Tasks;
 
-	const FString ImportScriptPath = LocalizationConfigurationScript::GetImportScriptPath(TargetSettings);
-	LocalizationConfigurationScript::GenerateImportScript(TargetSettings).Write(ImportScriptPath);
+	const FString ImportScriptPath = LocalizationConfigurationScript::GetImportScriptPath(TargetSettings, TOptional<FString>());
+	LocalizationConfigurationScript::GenerateImportScript(TargetSettings, TOptional<FString>(), DirectoryPath).Write(ImportScriptPath);
 	Tasks.Add(LocalizationCommandletExecution::FTask(LOCTEXT("ImportTaskName", "Import Translations"), ImportScriptPath));
 
 	const FString ReportScriptPath = LocalizationConfigurationScript::GetReportScriptPath(TargetSettings);
@@ -92,7 +93,7 @@ bool LocalizationCommandletTasks::ImportTarget(const TSharedRef<SWindow>& Parent
 	return LocalizationCommandletExecution::Execute(ParentWindow, WindowTitle, Tasks);
 }
 
-bool LocalizationCommandletTasks::ImportCulture(const TSharedRef<SWindow>& ParentWindow, FLocalizationTargetSettings& TargetSettings, const FString& CultureName)
+bool LocalizationCommandletTasks::ImportCulture(const TSharedRef<SWindow>& ParentWindow, FLocalizationTargetSettings& TargetSettings, const FString& CultureName, const TOptional<FString> FilePath)
 {
 	FCulturePtr Culture = FInternationalization::Get().GetCulture(CultureName);
 	if (!Culture.IsValid())
@@ -102,8 +103,9 @@ bool LocalizationCommandletTasks::ImportCulture(const TSharedRef<SWindow>& Paren
 
 	TArray<LocalizationCommandletExecution::FTask> Tasks;
 
-	const FString ImportScriptPath = LocalizationConfigurationScript::GetImportScriptPath(TargetSettings, &CultureName);
-	LocalizationConfigurationScript::GenerateImportScript(TargetSettings, &CultureName).Write(ImportScriptPath);
+	const FString DefaultImportScriptPath = LocalizationConfigurationScript::GetImportScriptPath(TargetSettings, TOptional<FString>(CultureName));
+	const FString ImportScriptPath = FPaths::CreateTempFilename(*FPaths::GetPath(DefaultImportScriptPath), *FPaths::GetBaseFilename(DefaultImportScriptPath), *FPaths::GetExtension(DefaultImportScriptPath, true));
+	LocalizationConfigurationScript::GenerateImportScript(TargetSettings, TOptional<FString>(CultureName), FilePath).Write(ImportScriptPath);
 	Tasks.Add(LocalizationCommandletExecution::FTask(LOCTEXT("ImportTaskName", "Import Translations"), ImportScriptPath));
 
 	const FString ReportScriptPath = LocalizationConfigurationScript::GetReportScriptPath(TargetSettings);
@@ -115,10 +117,14 @@ bool LocalizationCommandletTasks::ImportCulture(const TSharedRef<SWindow>& Paren
 	Arguments.Add(TEXT("TargetName"), FText::FromString(TargetSettings.Name));
 
 	const FText WindowTitle = FText::Format(LOCTEXT("ImportCultureForTargetWindowTitle", "Import {CultureName} Translations for Target {TargetName}"), Arguments);
-	return LocalizationCommandletExecution::Execute(ParentWindow, WindowTitle, Tasks);
+
+	bool HasSucceeeded = LocalizationCommandletExecution::Execute(ParentWindow, WindowTitle, Tasks);
+	IFileManager::Get().Delete(*ImportScriptPath); // Don't clutter up the loc config directory with scripts for individual cultures.
+	return HasSucceeeded;
+
 }
 
-bool LocalizationCommandletTasks::ExportTargets(const TSharedRef<SWindow>& ParentWindow, const TArray<FLocalizationTargetSettings*>& TargetsSettings)
+bool LocalizationCommandletTasks::ExportTargets(const TSharedRef<SWindow>& ParentWindow, const TArray<FLocalizationTargetSettings*>& TargetsSettings, const TOptional<FString> DirectoryPath)
 {
 	TArray<LocalizationCommandletExecution::FTask> Tasks;
 
@@ -128,20 +134,21 @@ bool LocalizationCommandletTasks::ExportTargets(const TSharedRef<SWindow>& Paren
 		Arguments.Add(TEXT("TargetName"), FText::FromString(TargetSettings->Name));
 
 		const FText ExportTaskName = FText::Format(LOCTEXT("ExportTaskNameFormat", "Export Translations for {TargetName}"), Arguments);
-		const FString ExportScriptPath = LocalizationConfigurationScript::GetExportScriptPath(*TargetSettings);
-		LocalizationConfigurationScript::GenerateExportScript(*TargetSettings).Write(ExportScriptPath);
+		const FString ExportScriptPath = LocalizationConfigurationScript::GetExportScriptPath(*TargetSettings, TOptional<FString>());
+		const TOptional<FString> DirectoryPathForTarget = DirectoryPath.IsSet() ? DirectoryPath.GetValue() / TargetSettings->Name : TOptional<FString>();
+		LocalizationConfigurationScript::GenerateExportScript(*TargetSettings, TOptional<FString>(), DirectoryPathForTarget).Write(ExportScriptPath);
 		Tasks.Add(LocalizationCommandletExecution::FTask(ExportTaskName, ExportScriptPath));
 	}
 
 	return LocalizationCommandletExecution::Execute(ParentWindow, LOCTEXT("ExportForAllTargetsWindowTitle", "Export Translations for All Targets"), Tasks);
 }
 
-bool LocalizationCommandletTasks::ExportTarget(const TSharedRef<SWindow>& ParentWindow, FLocalizationTargetSettings& TargetSettings)
+bool LocalizationCommandletTasks::ExportTarget(const TSharedRef<SWindow>& ParentWindow, FLocalizationTargetSettings& TargetSettings, const TOptional<FString> DirectoryPath)
 {
 	TArray<LocalizationCommandletExecution::FTask> Tasks;
 
-	const FString ExportScriptPath = LocalizationConfigurationScript::GetExportScriptPath(TargetSettings);
-	LocalizationConfigurationScript::GenerateExportScript(TargetSettings).Write(ExportScriptPath);
+	const FString ExportScriptPath = LocalizationConfigurationScript::GetExportScriptPath(TargetSettings, TOptional<FString>());
+	LocalizationConfigurationScript::GenerateExportScript(TargetSettings, TOptional<FString>(), DirectoryPath).Write(ExportScriptPath);
 	Tasks.Add(LocalizationCommandletExecution::FTask(LOCTEXT("ExportTaskName", "Export Translations"), ExportScriptPath));
 
 	FFormatNamedArguments Arguments;
@@ -151,7 +158,7 @@ bool LocalizationCommandletTasks::ExportTarget(const TSharedRef<SWindow>& Parent
 	return LocalizationCommandletExecution::Execute(ParentWindow, WindowTitle, Tasks);
 }
 
-bool LocalizationCommandletTasks::ExportCulture(const TSharedRef<SWindow>& ParentWindow, FLocalizationTargetSettings& TargetSettings, const FString& CultureName)
+bool LocalizationCommandletTasks::ExportCulture(const TSharedRef<SWindow>& ParentWindow, FLocalizationTargetSettings& TargetSettings, const FString& CultureName, const TOptional<FString> FilePath)
 {
 	FCulturePtr Culture = FInternationalization::Get().GetCulture(CultureName);
 	if (!Culture.IsValid())
@@ -161,8 +168,9 @@ bool LocalizationCommandletTasks::ExportCulture(const TSharedRef<SWindow>& Paren
 
 	TArray<LocalizationCommandletExecution::FTask> Tasks;
 
-	const FString ExportScriptPath = LocalizationConfigurationScript::GetExportScriptPath(TargetSettings, &CultureName);
-	LocalizationConfigurationScript::GenerateExportScript(TargetSettings, &CultureName).Write(ExportScriptPath);
+	const FString DefaultExportScriptPath = LocalizationConfigurationScript::GetExportScriptPath(TargetSettings, TOptional<FString>(CultureName));
+	const FString ExportScriptPath = FPaths::CreateTempFilename(*FPaths::GetPath(DefaultExportScriptPath), *FPaths::GetBaseFilename(DefaultExportScriptPath), *FPaths::GetExtension(DefaultExportScriptPath, true));
+	LocalizationConfigurationScript::GenerateExportScript(TargetSettings, TOptional<FString>(CultureName), FilePath).Write(ExportScriptPath);
 	Tasks.Add(LocalizationCommandletExecution::FTask(LOCTEXT("ExportTaskName", "Export Translations"), ExportScriptPath));
 
 	FFormatNamedArguments Arguments;
@@ -170,7 +178,10 @@ bool LocalizationCommandletTasks::ExportCulture(const TSharedRef<SWindow>& Paren
 	Arguments.Add(TEXT("TargetName"), FText::FromString(TargetSettings.Name));
 
 	const FText WindowTitle = FText::Format(LOCTEXT("ExportCultureForTargetWindowTitle", "Export {CultureName} Translations for Target {TargetName}"), Arguments);
-	return LocalizationCommandletExecution::Execute(ParentWindow, WindowTitle, Tasks);
+
+	bool HasSucceeeded = LocalizationCommandletExecution::Execute(ParentWindow, WindowTitle, Tasks);
+	IFileManager::Get().Delete(*ExportScriptPath); // Don't clutter up the loc config directory with scripts for individual cultures.
+	return HasSucceeeded;
 }
 
 bool LocalizationCommandletTasks::GenerateReportsForTargets(const TSharedRef<SWindow>& ParentWindow, const TArray<FLocalizationTargetSettings*>& TargetsSettings)
diff --git a/Engine/Source/Editor/LocalizationDashboard/Private/LocalizationCommandletTasks.h b/Engine/Source/Editor/LocalizationDashboard/Private/LocalizationCommandletTasks.h
--- a/Engine/Source/Editor/LocalizationDashboard/Private/LocalizationCommandletTasks.h
+++ b/Engine/Source/Editor/LocalizationDashboard/Private/LocalizationCommandletTasks.h
@@ -8,13 +8,13 @@ namespace LocalizationCommandletTasks
 	bool GatherTargets(const TSharedRef<SWindow>& ParentWindow, const TArray<FLocalizationTargetSettings*>& TargetsSettings);
 	bool GatherTarget(const TSharedRef<SWindow>& ParentWindow, FLocalizationTargetSettings& TargetSettings);
 
-	bool ImportTargets(const TSharedRef<SWindow>& ParentWindow, const TArray<FLocalizationTargetSettings*>& TargetsSettings);
-	bool ImportTarget(const TSharedRef<SWindow>& ParentWindow, FLocalizationTargetSettings& TargetSettings);
-	bool ImportCulture(const TSharedRef<SWindow>& ParentWindow, FLocalizationTargetSettings& TargetSettings, const FString& CultureName);
+	bool ImportTargets(const TSharedRef<SWindow>& ParentWindow, const TArray<FLocalizationTargetSettings*>& TargetsSettings, const TOptional<FString> DirectoryPath = TOptional<FString>());
+	bool ImportTarget(const TSharedRef<SWindow>& ParentWindow, FLocalizationTargetSettings& TargetSettings, const TOptional<FString> DirectoryPath = TOptional<FString>());
+	bool ImportCulture(const TSharedRef<SWindow>& ParentWindow, FLocalizationTargetSettings& TargetSettings, const FString& CultureName, const TOptional<FString> FilePath = TOptional<FString>());
 
-	bool ExportTargets(const TSharedRef<SWindow>& ParentWindow, const TArray<FLocalizationTargetSettings*>& TargetsSettings);
-	bool ExportTarget(const TSharedRef<SWindow>& ParentWindow, FLocalizationTargetSettings& TargetSettings);
-	bool ExportCulture(const TSharedRef<SWindow>& ParentWindow, FLocalizationTargetSettings& TargetSettings, const FString& CultureName);
+	bool ExportTargets(const TSharedRef<SWindow>& ParentWindow, const TArray<FLocalizationTargetSettings*>& TargetsSettings, const TOptional<FString> DirectoryPath = TOptional<FString>());
+	bool ExportTarget(const TSharedRef<SWindow>& ParentWindow, FLocalizationTargetSettings& TargetSettings, const TOptional<FString> DirectoryPath = TOptional<FString>());
+	bool ExportCulture(const TSharedRef<SWindow>& ParentWindow, FLocalizationTargetSettings& TargetSettings, const FString& CultureName, const TOptional<FString> FilePath = TOptional<FString>());
 
 	bool GenerateReportsForTargets(const TSharedRef<SWindow>& ParentWindow, const TArray<FLocalizationTargetSettings*>& TargetsSettings);
 	bool GenerateReportsForTarget(const TSharedRef<SWindow>& ParentWindow, FLocalizationTargetSettings& TargetSettings);
diff --git a/Engine/Source/Editor/LocalizationDashboard/Private/LocalizationConfigurationScript.cpp b/Engine/Source/Editor/LocalizationDashboard/Private/LocalizationConfigurationScript.cpp
--- a/Engine/Source/Editor/LocalizationDashboard/Private/LocalizationConfigurationScript.cpp
+++ b/Engine/Source/Editor/LocalizationDashboard/Private/LocalizationConfigurationScript.cpp
@@ -43,9 +43,14 @@ namespace LocalizationConfigurationScript
 		return GetDataDirectory(Target) / CultureName / FString::Printf( TEXT("%s.%s"), *Target.Name, TEXT("archive") );
 	}
 
-	FString GetPOPath(const FLocalizationTargetSettings& Target, const FString& CultureName)
+	FString GetDefaultPOFileName(const FLocalizationTargetSettings& Target)
 	{
-		return GetDataDirectory(Target) / CultureName / FString::Printf( TEXT("%s.%s"), *Target.Name, TEXT("po") );
+		return FString::Printf( TEXT("%s.%s"), *Target.Name, TEXT("po") );
+	}
+
+	FString GetDefaultPOPath(const FLocalizationTargetSettings& Target, const FString& CultureName)
+	{
+		return GetDataDirectory(Target) / CultureName / GetDefaultPOFileName(Target);
 	}
 
 	FString GetLocResPath(const FLocalizationTargetSettings& Target, const FString& CultureName)
@@ -170,7 +175,7 @@ namespace LocalizationConfigurationScript
 		return GetScriptDirectory() / FString::Printf( TEXT("%s_Gather.%s"), *(Target.Name), TEXT("ini") );
 	}
 
-	FLocalizationConfigurationScript GenerateImportScript(const FLocalizationTargetSettings& Target, const FString* const SpecificCulture)
+	FLocalizationConfigurationScript GenerateImportScript(const FLocalizationTargetSettings& Target, const TOptional<FString> CultureName, const TOptional<FString> OutputPathOverride)
 	{
 		FLocalizationConfigurationScript Script;
 
@@ -185,8 +190,28 @@ namespace LocalizationConfigurationScript
 
 			ConfigSection.Add( TEXT("bImportLoc"), TEXT("true") );
 
-			const FString SourcePath = ContentDirRelativeToGameDir / TEXT("Localization") / Target.Name;
+			FString SourcePath;
+			// Overriding output path changes the source directory for the PO file.
+			if (OutputPathOverride.IsSet())
+			{
+				// The output path for a specific culture is a file path.
+				if (CultureName.IsSet())
+				{
+					SourcePath = MakePathRelativeToProjectDirectory( FPaths::GetPath(OutputPathOverride.GetValue()) );
+				}
+				// Otherwise, it is a directory path.
+				else
+				{
+					SourcePath = MakePathRelativeToProjectDirectory( OutputPathOverride.GetValue() );
+				}
+			}
+			// Use the default PO file's directory path.
+			else
+			{
+				SourcePath = ContentDirRelativeToGameDir / TEXT("Localization") / Target.Name;
+			}
 			ConfigSection.Add( TEXT("SourcePath"), SourcePath );
+
 			const FString DestinationPath = ContentDirRelativeToGameDir / TEXT("Localization") / Target.Name;
 			ConfigSection.Add( TEXT("DestinationPath"), DestinationPath );
 
@@ -196,9 +221,9 @@ namespace LocalizationConfigurationScript
 			};
 
 			// Export for a specific culture.
-			if (SpecificCulture)
+			if (CultureName.IsSet())
 			{
-				ConfigSection.Add( TEXT("CulturesToGenerate"), *SpecificCulture );
+				ConfigSection.Add( TEXT("CulturesToGenerate"), CultureName.GetValue() );
 			}
 			// Export for all cultures.
 			else
@@ -209,9 +234,27 @@ namespace LocalizationConfigurationScript
 				}
 			}
 
+			// Do not use culture subdirectories if importing a single culture to a specific directory.
+			if (CultureName.IsSet() && OutputPathOverride.IsSet())
+			{
+				ConfigSection.Add( TEXT("bUseCultureDirectory"), "false" );
+			}
+
 			ConfigSection.Add( TEXT("ManifestName"), FPaths::GetCleanFilename(GetManifestPath(Target)) );
 			ConfigSection.Add( TEXT("ArchiveName"), FPaths::GetCleanFilename(GetArchivePath(Target, FString())) );
-			ConfigSection.Add( TEXT("PortableObjectName"), FPaths::GetCleanFilename(GetPOPath(Target, FString())) );
+
+			FString POFileName;
+			// The output path for a specific culture is a file path.
+			if (CultureName.IsSet() && OutputPathOverride.IsSet())
+			{
+				POFileName =  FPaths::GetCleanFilename( OutputPathOverride.GetValue() );
+			}
+			// Use the default PO file's name.
+			else
+			{
+				POFileName = FPaths::GetCleanFilename( GetDefaultPOFileName( Target ) );
+			}
+			ConfigSection.Add( TEXT("PortableObjectName"), POFileName );
 		}
 
 		Script.Dirty = true;
@@ -219,13 +262,13 @@ namespace LocalizationConfigurationScript
 		return Script;
 	}
 
-	FString GetImportScriptPath(const FLocalizationTargetSettings& Target, const FString* const SpecificCultureName)
+	FString GetImportScriptPath(const FLocalizationTargetSettings& Target, const TOptional<FString> CultureName)
 	{
 		const FString ConfigFileDirectory = GetScriptDirectory();
 		FString ConfigFilePath;
-		if (SpecificCultureName)
+		if (CultureName.IsSet())
 		{
-			ConfigFilePath = ConfigFileDirectory / FString::Printf( TEXT("%s_Import_%s.%s"), *Target.Name, **SpecificCultureName, TEXT("ini") );
+			ConfigFilePath = ConfigFileDirectory / FString::Printf( TEXT("%s_Import_%s.%s"), *Target.Name, *CultureName.GetValue(), TEXT("ini") );
 		}
 		else
 		{
@@ -234,7 +277,7 @@ namespace LocalizationConfigurationScript
 		return ConfigFilePath;
 	}
 
-	FLocalizationConfigurationScript GenerateExportScript(const FLocalizationTargetSettings& Target, const FString* const SpecificCulture)
+	FLocalizationConfigurationScript GenerateExportScript(const FLocalizationTargetSettings& Target, const TOptional<FString> CultureName, const TOptional<FString> OutputPathOverride)
 	{
 		FLocalizationConfigurationScript Script;
 
@@ -251,9 +294,30 @@ namespace LocalizationConfigurationScript
 
 			const FString SourcePath = ContentDirRelativeToGameDir / TEXT("Localization") / Target.Name;
 			ConfigSection.Add( TEXT("SourcePath"), SourcePath );
-			const FString DestinationPath = ContentDirRelativeToGameDir / TEXT("Localization") / Target.Name;
+
+			FString DestinationPath;
+			// Overriding output path changes the destination directory for the PO file.
+			if (OutputPathOverride.IsSet())
+			{
+				// The output path for a specific culture is a file path.
+				if (CultureName.IsSet())
+				{
+					DestinationPath = MakePathRelativeToProjectDirectory( FPaths::GetPath(OutputPathOverride.GetValue()) );
+				}
+				// Otherwise, it is a directory path.
+				else
+				{
+					DestinationPath = MakePathRelativeToProjectDirectory( OutputPathOverride.GetValue() );
+				}
+			}
+			// Use the default PO file's directory path.
+			else
+			{
+				DestinationPath = ContentDirRelativeToGameDir / TEXT("Localization") / Target.Name;
+			}
 			ConfigSection.Add( TEXT("DestinationPath"), DestinationPath );
 
+
 			TArray<const FCultureStatistics*> AllCultureStatistics;
 			AllCultureStatistics.Add(&Target.NativeCultureStatistics);
 			for (const FCultureStatistics& SupportedCultureStatistics : Target.SupportedCulturesStatistics)
@@ -267,9 +331,9 @@ namespace LocalizationConfigurationScript
 			};
 
 			// Export for a specific culture.
-			if (SpecificCulture)
+			if (CultureName.IsSet())
 			{
-				const int32 CultureIndex = AllCultureStatistics.IndexOfByPredicate([SpecificCulture](const FCultureStatistics* Culture) { return Culture->CultureName == *SpecificCulture; });
+				const int32 CultureIndex = AllCultureStatistics.IndexOfByPredicate([CultureName](const FCultureStatistics* Culture) { return Culture->CultureName == CultureName.GetValue(); });
 				AddCultureToGenerate(CultureIndex);
 			}
 			// Export for all cultures.
@@ -281,9 +345,27 @@ namespace LocalizationConfigurationScript
 				}
 			}
 
+			// Do not use culture subdirectories if exporting a single culture to a specific directory.
+			if (CultureName.IsSet() && OutputPathOverride.IsSet())
+			{
+				ConfigSection.Add( TEXT("bUseCultureDirectory"), "false" );
+			}
+
+
 			ConfigSection.Add( TEXT("ManifestName"), FPaths::GetCleanFilename(GetManifestPath(Target)) );
 			ConfigSection.Add( TEXT("ArchiveName"), FPaths::GetCleanFilename(GetArchivePath(Target, FString())) );
-			ConfigSection.Add( TEXT("PortableObjectName"), FPaths::GetCleanFilename(GetPOPath(Target, FString())) );
+			FString POFileName;
+			// The output path for a specific culture is a file path.
+			if (CultureName.IsSet() && OutputPathOverride.IsSet())
+			{
+				POFileName =  FPaths::GetCleanFilename( OutputPathOverride.GetValue() );
+			}
+			// Use the default PO file's name.
+			else
+			{
+				POFileName = FPaths::GetCleanFilename( GetDefaultPOPath( Target, CultureName.Get( TEXT("") ) ) );
+			}
+			ConfigSection.Add( TEXT("PortableObjectName"), POFileName );
 		}
 
 		Script.Dirty = true;
@@ -291,13 +373,13 @@ namespace LocalizationConfigurationScript
 		return Script;
 	}
 
-	FString GetExportScriptPath(const FLocalizationTargetSettings& Target, const FString* const SpecificCultureName)
+	FString GetExportScriptPath(const FLocalizationTargetSettings& Target, const TOptional<FString> CultureName)
 	{
 		const FString ConfigFileDirectory = GetScriptDirectory();
 		FString ConfigFilePath;
-		if (SpecificCultureName)
+		if (CultureName.IsSet())
 		{
-			ConfigFilePath = ConfigFileDirectory / FString::Printf( TEXT("%s_Export_%s.%s"), *Target.Name, **SpecificCultureName, TEXT("ini") );
+			ConfigFilePath = ConfigFileDirectory / FString::Printf( TEXT("%s_Export_%s.%s"), *Target.Name, *CultureName.GetValue(), TEXT("ini") );
 		}
 		else
 		{
diff --git a/Engine/Source/Editor/LocalizationDashboard/Private/LocalizationConfigurationScript.h b/Engine/Source/Editor/LocalizationDashboard/Private/LocalizationConfigurationScript.h
--- a/Engine/Source/Editor/LocalizationDashboard/Private/LocalizationConfigurationScript.h
+++ b/Engine/Source/Editor/LocalizationDashboard/Private/LocalizationConfigurationScript.h
@@ -28,17 +28,18 @@ namespace LocalizationConfigurationScript
 
 	FString GetManifestPath(const FLocalizationTargetSettings& Target);
 	FString GetArchivePath(const FLocalizationTargetSettings& Target, const FString& CultureName);
-	FString GetPOPath(const FLocalizationTargetSettings& Target, const FString& CultureName);
+	FString GetDefaultPOFileName(const FLocalizationTargetSettings& Target);
+	FString GetDefaultPOPath(const FLocalizationTargetSettings& Target, const FString& CultureName);
 	FString GetLocResPath(const FLocalizationTargetSettings& Target, const FString& CultureName);
 
 	FLocalizationConfigurationScript GenerateGatherScript(const FLocalizationTargetSettings& Target);
 	FString GetGatherScriptPath(const FLocalizationTargetSettings& Target);
 
-	FLocalizationConfigurationScript GenerateImportScript(const FLocalizationTargetSettings& Target, const FString* const SpecificCulture = nullptr);
-	FString GetImportScriptPath(const FLocalizationTargetSettings& Target, const FString* const SpecificCultureName = nullptr);
+	FLocalizationConfigurationScript GenerateImportScript(const FLocalizationTargetSettings& Target, const TOptional<FString> CultureName = TOptional<FString>(), const TOptional<FString> OutputPathOverride = TOptional<FString>());
+	FString GetImportScriptPath(const FLocalizationTargetSettings& Target, const TOptional<FString> CultureName = TOptional<FString>());
 
-	FLocalizationConfigurationScript GenerateExportScript(const FLocalizationTargetSettings& Target, const FString* const SpecificCulture = nullptr);
-	FString GetExportScriptPath(const FLocalizationTargetSettings& Target, const FString* const SpecificCultureName = nullptr);
+	FLocalizationConfigurationScript GenerateExportScript(const FLocalizationTargetSettings& Target, const TOptional<FString> CultureName = TOptional<FString>(), const TOptional<FString> OutputPathOverride = TOptional<FString>());
+	FString GetExportScriptPath(const FLocalizationTargetSettings& Target, const TOptional<FString> CultureName = TOptional<FString>());
 
 	FLocalizationConfigurationScript GenerateReportScript(const FLocalizationTargetSettings& Target);
 	FString GetReportScriptPath(const FLocalizationTargetSettings& Target);
diff --git a/Engine/Source/Editor/LocalizationDashboard/Private/LocalizationDashboard.cpp b/Engine/Source/Editor/LocalizationDashboard/Private/LocalizationDashboard.cpp
--- a/Engine/Source/Editor/LocalizationDashboard/Private/LocalizationDashboard.cpp
+++ b/Engine/Source/Editor/LocalizationDashboard/Private/LocalizationDashboard.cpp
@@ -387,29 +387,64 @@ namespace
 
 	void FProjectLocalizationSettingsDetailsCustomization::ImportAllTargets()
 	{
-		const TSharedPtr<SWindow> ParentWindow = FSlateApplication::Get().FindWidgetWindow(DetailLayoutBuilder->GetDetailsView().AsShared());
-		TArray<FLocalizationTargetSettings*> TargetSettings;
-		for (ULocalizationTarget* const TargetObject : Settings->TargetObjects)
+		IDesktopPlatform* DesktopPlatform = FDesktopPlatformModule::Get();
+		if (DesktopPlatform)
 		{
-			TargetSettings.Add(&TargetObject->Settings);
-		}
-		LocalizationCommandletTasks::ImportTargets(ParentWindow.ToSharedRef(), TargetSettings);
+			void* ParentWindowWindowHandle = NULL;
+			const TSharedPtr<SWindow> ParentWindow = FSlateApplication::Get().FindWidgetWindow(DetailLayoutBuilder->GetDetailsView().AsShared());
+			if (ParentWindow.IsValid() && ParentWindow->GetNativeWindow().IsValid())
+			{
+				ParentWindowWindowHandle = ParentWindow->GetNativeWindow()->GetOSWindowHandle();
+			}
 
-		for (ULocalizationTarget* const LocalizationTarget : Settings->TargetObjects)
-		{
-			UpdateTargetFromReports(LocalizationTarget);
+			
+			const FString DefaultPath = FPaths::ConvertRelativePathToFull(FPaths::GameContentDir() / TEXT("Localization"));
+
+			// Prompt the user for the directory
+			FString OutputDirectory;
+			if (DesktopPlatform->OpenDirectoryDialog(ParentWindowWindowHandle, LOCTEXT("ImportAllTranslationsDialogTitle", "Import All Translations from Directory").ToString(), DefaultPath, OutputDirectory))
+			{
+				TArray<FLocalizationTargetSettings*> TargetSettings;
+				for (ULocalizationTarget* const TargetObject : Settings->TargetObjects)
+				{
+					TargetSettings.Add(&TargetObject->Settings);
+				}
+				LocalizationCommandletTasks::ImportTargets(ParentWindow.ToSharedRef(), TargetSettings, TOptional<FString>(OutputDirectory));
+
+				for (ULocalizationTarget* const LocalizationTarget : Settings->TargetObjects)
+				{
+					UpdateTargetFromReports(LocalizationTarget);
+				}
+			}
 		}
 	}
 
 	void FProjectLocalizationSettingsDetailsCustomization::ExportAllTargets()
 	{
-		const TSharedPtr<SWindow> ParentWindow = FSlateApplication::Get().FindWidgetWindow(DetailLayoutBuilder->GetDetailsView().AsShared());
-		TArray<FLocalizationTargetSettings*> TargetSettings;
-		for (ULocalizationTarget* const TargetObject : Settings->TargetObjects)
+		IDesktopPlatform* DesktopPlatform = FDesktopPlatformModule::Get();
+		if (DesktopPlatform)
 		{
-			TargetSettings.Add(&TargetObject->Settings);
+			void* ParentWindowWindowHandle = NULL;
+			const TSharedPtr<SWindow> ParentWindow = FSlateApplication::Get().FindWidgetWindow(DetailLayoutBuilder->GetDetailsView().AsShared());
+			if (ParentWindow.IsValid() && ParentWindow->GetNativeWindow().IsValid())
+			{
+				ParentWindowWindowHandle = ParentWindow->GetNativeWindow()->GetOSWindowHandle();
+			}
+
+			const FString DefaultPath = FPaths::ConvertRelativePathToFull(FPaths::GameContentDir() / TEXT("Localization"));
+
+			// Prompt the user for the directory
+			FString OutputDirectory;
+			if (DesktopPlatform->OpenDirectoryDialog(ParentWindowWindowHandle, LOCTEXT("ExportAllTranslationsDialogTitle", "Export All Translations to Directory").ToString(), DefaultPath, OutputDirectory))
+			{
+				TArray<FLocalizationTargetSettings*> TargetSettings;
+				for (ULocalizationTarget* const TargetObject : Settings->TargetObjects)
+				{
+					TargetSettings.Add(&TargetObject->Settings);
+				}
+				LocalizationCommandletTasks::ExportTargets(ParentWindow.ToSharedRef(), TargetSettings, TOptional<FString>(OutputDirectory));
+			}
 		}
-		LocalizationCommandletTasks::ExportTargets(ParentWindow.ToSharedRef(), TargetSettings);
 	}
 
 	void FProjectLocalizationSettingsDetailsCustomization::UpdateTargetFromReports(ULocalizationTarget* const LocalizationTarget)
diff --git a/Engine/Source/Editor/LocalizationDashboard/Private/LocalizationTarget.cpp b/Engine/Source/Editor/LocalizationDashboard/Private/LocalizationTarget.cpp
--- a/Engine/Source/Editor/LocalizationDashboard/Private/LocalizationTarget.cpp
+++ b/Engine/Source/Editor/LocalizationDashboard/Private/LocalizationTarget.cpp
@@ -255,7 +255,7 @@ bool FLocalizationTargetSettings::Rename(const FString& NewName)
 		for (const FCultureStatistics* Culture : Cultures)
 		{
 			NamedPaths.Add(LocalizationConfigurationScript::GetArchivePath(*this, Culture->CultureName));
-			NamedPaths.Add(LocalizationConfigurationScript::GetPOPath(*this, Culture->CultureName));
+			NamedPaths.Add(LocalizationConfigurationScript::GetDefaultPOPath(*this, Culture->CultureName));
 			NamedPaths.Add(LocalizationConfigurationScript::GetLocResPath(*this, Culture->CultureName));
 		}
 
diff --git a/Engine/Source/Editor/LocalizationDashboard/Private/SLocalizationDashboardTargetRow.cpp b/Engine/Source/Editor/LocalizationDashboard/Private/SLocalizationDashboardTargetRow.cpp
--- a/Engine/Source/Editor/LocalizationDashboard/Private/SLocalizationDashboardTargetRow.cpp
+++ b/Engine/Source/Editor/LocalizationDashboard/Private/SLocalizationDashboardTargetRow.cpp
@@ -342,12 +342,32 @@ FReply SLocalizationDashboardTargetRow::Gather()
 FReply SLocalizationDashboardTargetRow::ImportAll()
 {
 	ULocalizationTarget* const LocalizationTarget = GetTarget();
-	if (LocalizationTarget)
+	IDesktopPlatform* DesktopPlatform = FDesktopPlatformModule::Get();
+	if (LocalizationTarget && DesktopPlatform)
 	{
+		void* ParentWindowWindowHandle = NULL;
 		const TSharedPtr<SWindow> ParentWindow = FSlateApplication::Get().FindWidgetWindow(AsShared());
-		LocalizationCommandletTasks::ImportTarget(ParentWindow.ToSharedRef(), LocalizationTarget->Settings);
+		if (ParentWindow.IsValid() && ParentWindow->GetNativeWindow().IsValid())
+		{
+			ParentWindowWindowHandle = ParentWindow->GetNativeWindow()->GetOSWindowHandle();
+		}
 
-		UpdateTargetFromReports();
+		const FString DefaultPath = FPaths::ConvertRelativePathToFull(LocalizationConfigurationScript::GetDataDirectory(LocalizationTarget->Settings));
+
+		FText DialogTitle;
+		{
+			FFormatNamedArguments FormatArguments;
+			FormatArguments.Add(TEXT("TargetName"), FText::FromString(LocalizationTarget->Settings.Name));
+			DialogTitle = FText::Format(LOCTEXT("ImportAllTranslationsForTargetDialogTitleFormat", "Import All Translations for {TargetName} from Directory"), FormatArguments);
+		}
+
+		// Prompt the user for the directory
+		FString OutputDirectory;
+		if (DesktopPlatform->OpenDirectoryDialog(ParentWindowWindowHandle, DialogTitle.ToString(), DefaultPath, OutputDirectory))
+		{
+			LocalizationCommandletTasks::ImportTarget(ParentWindow.ToSharedRef(), LocalizationTarget->Settings, TOptional<FString>(OutputDirectory));
+			UpdateTargetFromReports();
+		}
 	}
 
 	return FReply::Handled();
@@ -356,10 +376,31 @@ FReply SLocalizationDashboardTargetRow::ImportAll()
 FReply SLocalizationDashboardTargetRow::ExportAll()
 {
 	ULocalizationTarget* const LocalizationTarget = GetTarget();
-	if (LocalizationTarget)
+	IDesktopPlatform* DesktopPlatform = FDesktopPlatformModule::Get();
+	if (LocalizationTarget && DesktopPlatform)
 	{
+		void* ParentWindowWindowHandle = NULL;
 		const TSharedPtr<SWindow> ParentWindow = FSlateApplication::Get().FindWidgetWindow(AsShared());
-		LocalizationCommandletTasks::ExportTarget(ParentWindow.ToSharedRef(), LocalizationTarget->Settings);
+		if (ParentWindow.IsValid() && ParentWindow->GetNativeWindow().IsValid())
+		{
+			ParentWindowWindowHandle = ParentWindow->GetNativeWindow()->GetOSWindowHandle();
+		}
+
+		const FString DefaultPath = FPaths::ConvertRelativePathToFull(LocalizationConfigurationScript::GetDataDirectory(LocalizationTarget->Settings));
+
+		FText DialogTitle;
+		{
+			FFormatNamedArguments FormatArguments;
+			FormatArguments.Add(TEXT("TargetName"), FText::FromString(LocalizationTarget->Settings.Name));
+			DialogTitle = FText::Format(LOCTEXT("ExportAllTranslationsForTargetDialogTitleFormat", "Export All Translations for {TargetName} to Directory"), FormatArguments);
+		}
+
+		// Prompt the user for the directory
+		FString OutputDirectory;
+		if (DesktopPlatform->OpenDirectoryDialog(ParentWindowWindowHandle, DialogTitle.ToString(), DefaultPath, OutputDirectory))
+		{
+			LocalizationCommandletTasks::ExportTarget(ParentWindow.ToSharedRef(), LocalizationTarget->Settings, TOptional<FString>(OutputDirectory));
+		}
 	}
 
 	return FReply::Handled();
diff --git a/Engine/Source/Editor/LocalizationDashboard/Private/SLocalizationTargetEditor.cpp b/Engine/Source/Editor/LocalizationDashboard/Private/SLocalizationTargetEditor.cpp
--- a/Engine/Source/Editor/LocalizationDashboard/Private/SLocalizationTargetEditor.cpp
+++ b/Engine/Source/Editor/LocalizationDashboard/Private/SLocalizationTargetEditor.cpp
@@ -672,21 +672,63 @@ namespace
 
 	void FLocalizationTargetDetailCustomization::ImportAllCultures()
 	{
-		if (LocalizationTarget)
+		IDesktopPlatform* DesktopPlatform = FDesktopPlatformModule::Get();
+		if (LocalizationTarget && DesktopPlatform)
 		{
+			void* ParentWindowWindowHandle = NULL;
 			const TSharedPtr<SWindow> ParentWindow = FSlateApplication::Get().FindWidgetWindow(DetailLayoutBuilder->GetDetailsView().AsShared());
-			LocalizationCommandletTasks::ImportTarget(ParentWindow.ToSharedRef(), LocalizationTarget->Settings);
+			if (ParentWindow.IsValid() && ParentWindow->GetNativeWindow().IsValid())
+			{
+				ParentWindowWindowHandle = ParentWindow->GetNativeWindow()->GetOSWindowHandle();
+			}
 
-			UpdateTargetFromReports();
+			const FString DefaultPath = FPaths::ConvertRelativePathToFull(LocalizationConfigurationScript::GetDataDirectory(LocalizationTarget->Settings));
+
+			FText DialogTitle;
+			{
+				FFormatNamedArguments FormatArguments;
+				FormatArguments.Add(TEXT("TargetName"), FText::FromString(LocalizationTarget->Settings.Name));
+				DialogTitle = FText::Format(LOCTEXT("ImportAllTranslationsForTargetDialogTitleFormat", "Import All Translations for {TargetName} from Directory"), FormatArguments);
+			}
+
+			// Prompt the user for the directory
+			FString OutputDirectory;
+			if (DesktopPlatform->OpenDirectoryDialog(ParentWindowWindowHandle, DialogTitle.ToString(), DefaultPath, OutputDirectory))
+			{
+				LocalizationCommandletTasks::ImportTarget(ParentWindow.ToSharedRef(), LocalizationTarget->Settings, TOptional<FString>(OutputDirectory));
+
+				UpdateTargetFromReports();
+			}
 		}
 	}
 
 	void FLocalizationTargetDetailCustomization::ExportAllCultures()
 	{
-		if (LocalizationTarget)
+		IDesktopPlatform* DesktopPlatform = FDesktopPlatformModule::Get();
+		if (LocalizationTarget && DesktopPlatform)
 		{
+			void* ParentWindowWindowHandle = NULL;
 			const TSharedPtr<SWindow> ParentWindow = FSlateApplication::Get().FindWidgetWindow(DetailLayoutBuilder->GetDetailsView().AsShared());
-			LocalizationCommandletTasks::ExportTarget(ParentWindow.ToSharedRef(), LocalizationTarget->Settings);
+			if (ParentWindow.IsValid() && ParentWindow->GetNativeWindow().IsValid())
+			{
+				ParentWindowWindowHandle = ParentWindow->GetNativeWindow()->GetOSWindowHandle();
+			}
+
+			const FString DefaultPath = FPaths::ConvertRelativePathToFull(LocalizationConfigurationScript::GetDataDirectory(LocalizationTarget->Settings));
+
+			FText DialogTitle;
+			{
+				FFormatNamedArguments FormatArguments;
+				FormatArguments.Add(TEXT("TargetName"), FText::FromString(LocalizationTarget->Settings.Name));
+				DialogTitle = FText::Format(LOCTEXT("ExportAllTranslationsForTargetDialogTitleFormat", "Export All Translations for {TargetName} to Directory"), FormatArguments);
+			}
+
+			// Prompt the user for the directory
+			FString OutputDirectory;
+			if (DesktopPlatform->OpenDirectoryDialog(ParentWindowWindowHandle, DialogTitle.ToString(), DefaultPath, OutputDirectory))
+			{
+				LocalizationCommandletTasks::ExportTarget(ParentWindow.ToSharedRef(), LocalizationTarget->Settings, TOptional<FString>(OutputDirectory));
+			}
 		}
 	}
 
diff --git a/Engine/Source/Editor/LocalizationDashboard/Private/SLocalizationTargetEditorCultureRow.cpp b/Engine/Source/Editor/LocalizationDashboard/Private/SLocalizationTargetEditorCultureRow.cpp
--- a/Engine/Source/Editor/LocalizationDashboard/Private/SLocalizationTargetEditorCultureRow.cpp
+++ b/Engine/Source/Editor/LocalizationDashboard/Private/SLocalizationTargetEditorCultureRow.cpp
@@ -302,20 +302,45 @@ FReply SLocalizationTargetEditorCultureRow::Edit()
 
 FReply SLocalizationTargetEditorCultureRow::Import()
 {
+	const FCulturePtr Culture = GetCulture();
 	FLocalizationTargetSettings* const TargetSettings = GetTargetSettings();
-	if (TargetSettings)
+	IDesktopPlatform* DesktopPlatform = FDesktopPlatformModule::Get();
+	if (Culture.IsValid() && TargetSettings && DesktopPlatform)
 	{
-		const FCulturePtr Culture = GetCulture();
-		if (Culture.IsValid())
+		void* ParentWindowWindowHandle = NULL;
+		const TSharedPtr<SWindow> ParentWindow = FSlateApplication::Get().FindWidgetWindow(AsShared());
+		if (ParentWindow.IsValid() && ParentWindow->GetNativeWindow().IsValid())
+		{
+			ParentWindowWindowHandle = ParentWindow->GetNativeWindow()->GetOSWindowHandle();
+		}
+
+		const FString POFileName = LocalizationConfigurationScript::GetDefaultPOFileName(*TargetSettings);
+		const FString POFileTypeDescription = LOCTEXT("PortableObjectFileDescription", "Portable Object").ToString();
+		const FString POFileExtension = FPaths::GetExtension(POFileName);
+		const FString POFileExtensionWildcard = FString::Printf(TEXT("*.%s"), *POFileExtension);
+		const FString FileTypes = FString::Printf(TEXT("%s (%s)|%s"), *POFileTypeDescription, *POFileExtensionWildcard, *POFileExtensionWildcard);
+		const FString DefaultFilename = POFileName;
+		const FString DefaultPath = FPaths::GetPath(LocalizationConfigurationScript::GetDefaultPOPath(*TargetSettings, Culture->GetName()));
+
+		FText DialogTitle;
+		{
+			FFormatNamedArguments FormatArguments;
+			FormatArguments.Add(TEXT("TargetName"), FText::FromString(TargetSettings->Name));
+			FormatArguments.Add(TEXT("CultureName"), FText::FromString(Culture->GetDisplayName()));
+			FText DialogTitle = FText::Format(LOCTEXT("ImportSpecificTranslationsForTargetDialogTitleFormat", "Import {CultureName} Translations for {TargetName} from Directory"), FormatArguments);
+		}
+
+		// Prompt the user for the directory
+		TArray<FString> OpenFilenames;
+		if (DesktopPlatform->OpenFileDialog(ParentWindowWindowHandle, DialogTitle.ToString(), DefaultPath, DefaultFilename, FileTypes, 0, OpenFilenames))
 		{
 			const TSharedPtr<SWindow> ParentWindow = FSlateApplication::Get().FindWidgetWindow(AsShared());
-			LocalizationCommandletTasks::ImportCulture(ParentWindow.ToSharedRef(), *TargetSettings, Culture->GetName());
+			LocalizationCommandletTasks::ImportCulture(ParentWindow.ToSharedRef(), *TargetSettings, Culture->GetName(), TOptional<FString>(OpenFilenames.Top()));
 
 			UpdateTargetFromReports();
 		}
 	}
 
-
 	return FReply::Handled();
 }
 
@@ -323,10 +348,38 @@ FReply SLocalizationTargetEditorCultureRow::Export()
 {
 	const FCulturePtr Culture = GetCulture();
 	FLocalizationTargetSettings* const TargetSettings = GetTargetSettings();
-	if (Culture.IsValid() && TargetSettings)
+	IDesktopPlatform* DesktopPlatform = FDesktopPlatformModule::Get();
+	if (Culture.IsValid() && TargetSettings && DesktopPlatform)
 	{
+		void* ParentWindowWindowHandle = NULL;
 		const TSharedPtr<SWindow> ParentWindow = FSlateApplication::Get().FindWidgetWindow(AsShared());
-		LocalizationCommandletTasks::ExportCulture(ParentWindow.ToSharedRef(), *TargetSettings, Culture->GetName());
+		if (ParentWindow.IsValid() && ParentWindow->GetNativeWindow().IsValid())
+		{
+			ParentWindowWindowHandle = ParentWindow->GetNativeWindow()->GetOSWindowHandle();
+		}
+
+		const FString POFileName = LocalizationConfigurationScript::GetDefaultPOFileName(*TargetSettings);
+		const FString POFileTypeDescription = LOCTEXT("PortableObjectFileDescription", "Portable Object").ToString();
+		const FString POFileExtension = FPaths::GetExtension(POFileName);
+		const FString POFileExtensionWildcard = FString::Printf(TEXT("*.%s"), *POFileExtension);
+		const FString FileTypes = FString::Printf(TEXT("%s (%s)|%s"), *POFileTypeDescription, *POFileExtensionWildcard, *POFileExtensionWildcard);
+		const FString DefaultFilename = POFileName;
+		const FString DefaultPath = FPaths::GetPath(LocalizationConfigurationScript::GetDefaultPOPath(*TargetSettings, Culture->GetName()));
+
+		FText DialogTitle;
+		{
+			FFormatNamedArguments FormatArguments;
+			FormatArguments.Add(TEXT("TargetName"), FText::FromString(TargetSettings->Name));
+			FormatArguments.Add(TEXT("CultureName"), FText::FromString(Culture->GetDisplayName()));
+			DialogTitle = FText::Format(LOCTEXT("ExportSpecificTranslationsForTargetDialogTitleFormat", "Export {CultureName} Translations for {TargetName} to Directory"), FormatArguments);
+		}
+
+		// Prompt the user for the directory
+		TArray<FString> SaveFilenames;
+		if (DesktopPlatform->SaveFileDialog(ParentWindowWindowHandle, DialogTitle.ToString(), DefaultPath, DefaultFilename, FileTypes, 0, SaveFilenames))
+		{
+			LocalizationCommandletTasks::ExportCulture(ParentWindow.ToSharedRef(), *TargetSettings, Culture->GetName(), TOptional<FString>(SaveFilenames.Top()));
+		}
 	}
 
 	return FReply::Handled();
