Added buttons to localization dashboard and target editor to "compile translations" (generate localization resource files).

diff --git a/Engine/Source/Editor/EditorStyle/Private/SlateEditorStyle.cpp b/Engine/Source/Editor/EditorStyle/Private/SlateEditorStyle.cpp
--- a/Engine/Source/Editor/EditorStyle/Private/SlateEditorStyle.cpp
+++ b/Engine/Source/Editor/EditorStyle/Private/SlateEditorStyle.cpp
@@ -6304,17 +6304,20 @@ void FSlateEditorStyle::FStyle::SetupLocalizationDashboardStyles()
 	Set("LocalizationDashboard.GatherAllTargets", new IMAGE_BRUSH("Icons/icon_Refresh_16x", Icon40x40));
 	Set("LocalizationDashboard.ImportForAllTargetsCultures", new IMAGE_BRUSH("Icons/icon_Samples_Download_40x", Icon40x40));
 	Set("LocalizationDashboard.ExportForAllTargetsCultures", new IMAGE_BRUSH("Icons/icon_Import_40x", Icon40x40));
+	Set("LocalizationDashboard.CompileAllTargets", new IMAGE_BRUSH("Icons/icon_package_16x", Icon40x40));
 
 	Set("LocalizationDashboard.GatherTarget", new IMAGE_BRUSH("Icons/icon_Refresh_16x", Icon16x16));
 	Set("LocalizationDashboard.RefreshWordCount", new IMAGE_BRUSH("Icons/icon_Refresh_16x", Icon16x16));
 	Set("LocalizationDashboard.ImportForAllCultures", new IMAGE_BRUSH("Icons/icon_Samples_Download_20x", Icon16x16));
 	Set("LocalizationDashboard.ExportForAllCultures", new IMAGE_BRUSH("Icons/icon_Import_16x", Icon16x16));
+	Set("LocalizationDashboard.CompileTarget", new IMAGE_BRUSH("Icons/icon_package_16x", Icon16x16));
 	Set("LocalizationDashboard.DeleteTarget", new IMAGE_BRUSH("Icons/Cross_12x", Icon12x12 ) );
 
 	Set("LocalizationTargetEditor.Gather", new IMAGE_BRUSH("Icons/icon_Refresh_16x", Icon40x40));
 	Set("LocalizationTargetEditor.ImportForAllCultures", new IMAGE_BRUSH("Icons/icon_Samples_Download_40x", Icon40x40));
 	Set("LocalizationTargetEditor.ExportForAllCultures", new IMAGE_BRUSH("Icons/icon_Import_40x", Icon40x40));
 	Set("LocalizationTargetEditor.RefreshWordCounts", new IMAGE_BRUSH("Icons/icon_Refresh_16x", Icon40x40));
+	Set("LocalizationTargetEditor.Compile", new IMAGE_BRUSH( "Icons/icon_package_16x", Icon40x40));
 
 	Set("LocalizationTargetEditor.NativeCulture", new IMAGE_BRUSH( "Icons/Star_16x", Icon16x16 ) );
 	Set("LocalizationTargetEditor.EditTranslations", new IMAGE_BRUSH("Icons/icon_file_open_16px", Icon16x16));
@@ -6322,7 +6325,6 @@ void FSlateEditorStyle::FStyle::SetupLocalizationDashboardStyles()
 	Set("LocalizationTargetEditor.ExportForCulture", new IMAGE_BRUSH("Icons/icon_Import_16x", Icon16x16));
 	Set("LocalizationTargetEditor.DeleteCulture", new IMAGE_BRUSH("Icons/Cross_12x", Icon12x12 ) );
 
-	Set( "LocalizationDashboard.CommandletLog.ListBorder", new BOX_BRUSH( "/Docking/AppTabContentArea", FMargin(4/16.0f) ) );
 	Set( "LocalizationDashboard.CommandletLog.Text", FTextBlockStyle(NormalText)
 		.SetFont( TTF_CORE_FONT( "Fonts/Roboto-Regular", 8 ) )
 		.SetShadowOffset( FVector2D::ZeroVector )
diff --git a/Engine/Source/Editor/LocalizationDashboard/Private/LocalizationCommandletTasks.cpp b/Engine/Source/Editor/LocalizationDashboard/Private/LocalizationCommandletTasks.cpp
--- a/Engine/Source/Editor/LocalizationDashboard/Private/LocalizationCommandletTasks.cpp
+++ b/Engine/Source/Editor/LocalizationDashboard/Private/LocalizationCommandletTasks.cpp
@@ -217,4 +217,37 @@ bool LocalizationCommandletTasks::GenerateReportsForTarget(const TSharedRef<SWin
 	return LocalizationCommandletExecution::Execute(ParentWindow, WindowTitle, Tasks);
 }
 
+bool LocalizationCommandletTasks::CompileTargets(const TSharedRef<SWindow>& ParentWindow, const TArray<FLocalizationTargetSettings*>& TargetsSettings)
+{
+	TArray<LocalizationCommandletExecution::FTask> Tasks;
+
+	for (FLocalizationTargetSettings* TargetSettings : TargetsSettings)
+	{
+		FFormatNamedArguments Arguments;
+		Arguments.Add(TEXT("TargetName"), FText::FromString(TargetSettings->Name));
+
+		const FText CompileTaskName = FText::Format(LOCTEXT("CompileTaskNameFormat", "Compile Translations for {TargetName}"), Arguments);
+		const FString CompileScriptPath = LocalizationConfigurationScript::GetCompileScriptPath(*TargetSettings);
+		LocalizationConfigurationScript::GenerateCompileScript(*TargetSettings).Write(CompileScriptPath);
+		Tasks.Add(LocalizationCommandletExecution::FTask(CompileTaskName, CompileScriptPath));
+	}
+
+	return LocalizationCommandletExecution::Execute(ParentWindow, LOCTEXT("GenerateLocResForAllTargetsWindowTitle", "Compile Translations for All Targets"), Tasks);
+}
+
+bool LocalizationCommandletTasks::CompileTarget(const TSharedRef<SWindow>& ParentWindow, FLocalizationTargetSettings& TargetSettings)
+{
+	TArray<LocalizationCommandletExecution::FTask> Tasks;
+
+	const FString CompileScriptPath = LocalizationConfigurationScript::GetCompileScriptPath(TargetSettings);
+	LocalizationConfigurationScript::GenerateCompileScript(TargetSettings).Write(CompileScriptPath);
+	Tasks.Add(LocalizationCommandletExecution::FTask(LOCTEXT("CompileTaskName", "Compile Translations"), CompileScriptPath));
+
+	FFormatNamedArguments Arguments;
+	Arguments.Add(TEXT("TargetName"), FText::FromString(TargetSettings.Name));
+
+	const FText WindowTitle = FText::Format(LOCTEXT("GenerateLocResForTargetWindowTitle", "Compile Translations for Target {TargetName}"), Arguments);
+	return LocalizationCommandletExecution::Execute(ParentWindow, WindowTitle, Tasks);
+}
+
 #undef LOCTEXT_NAMESPACE
\ No newline at end of file
diff --git a/Engine/Source/Editor/LocalizationDashboard/Private/LocalizationCommandletTasks.h b/Engine/Source/Editor/LocalizationDashboard/Private/LocalizationCommandletTasks.h
--- a/Engine/Source/Editor/LocalizationDashboard/Private/LocalizationCommandletTasks.h
+++ b/Engine/Source/Editor/LocalizationDashboard/Private/LocalizationCommandletTasks.h
@@ -18,4 +18,7 @@ namespace LocalizationCommandletTasks
 
 	bool GenerateReportsForTargets(const TSharedRef<SWindow>& ParentWindow, const TArray<FLocalizationTargetSettings*>& TargetsSettings);
 	bool GenerateReportsForTarget(const TSharedRef<SWindow>& ParentWindow, FLocalizationTargetSettings& TargetSettings);
+
+	bool CompileTargets(const TSharedRef<SWindow>& ParentWindow, const TArray<FLocalizationTargetSettings*>& TargetsSettings);
+	bool CompileTarget(const TSharedRef<SWindow>& ParentWindow, FLocalizationTargetSettings& TargetSettings);
 }
\ No newline at end of file
diff --git a/Engine/Source/Editor/LocalizationDashboard/Private/LocalizationConfigurationScript.cpp b/Engine/Source/Editor/LocalizationDashboard/Private/LocalizationConfigurationScript.cpp
--- a/Engine/Source/Editor/LocalizationDashboard/Private/LocalizationConfigurationScript.cpp
+++ b/Engine/Source/Editor/LocalizationDashboard/Private/LocalizationConfigurationScript.cpp
@@ -58,6 +58,16 @@ namespace LocalizationConfigurationScript
 		return GetDataDirectory(Target) / CultureName / FString::Printf( TEXT("%s.%s"), *Target.Name, TEXT("locres") );
 	}
 
+	FString GetWordCountCSVPath(const FLocalizationTargetSettings& Target)
+	{
+		return GetDataDirectory(Target) / FString::Printf( TEXT("%s.%s"), *Target.Name, TEXT("csv") );
+	}
+
+	FString GetConflictReportPath(const FLocalizationTargetSettings& Target)
+	{
+		return GetDataDirectory(Target) / FString::Printf( TEXT("%s_Conflicts.%s"), *Target.Name, TEXT("txt") );
+	}
+
 	FLocalizationConfigurationScript GenerateGatherScript(const FLocalizationTargetSettings& Target)
 	{
 		FLocalizationConfigurationScript Script;
@@ -317,7 +327,6 @@ namespace LocalizationConfigurationScript
 			}
 			ConfigSection.Add( TEXT("DestinationPath"), DestinationPath );
 
-
 			TArray<const FCultureStatistics*> AllCultureStatistics;
 			AllCultureStatistics.Add(&Target.NativeCultureStatistics);
 			for (const FCultureStatistics& SupportedCultureStatistics : Target.SupportedCulturesStatistics)
@@ -432,13 +441,47 @@ namespace LocalizationConfigurationScript
 		return GetScriptDirectory() / FString::Printf( TEXT("%s_GenerateReports.%s"), *(Target.Name), TEXT("ini") );
 	}
 
-	FString GetWordCountCSVPath(const FLocalizationTargetSettings& Target)
+	FLocalizationConfigurationScript GenerateCompileScript(const FLocalizationTargetSettings& Target)
 	{
-		return GetDataDirectory(Target) / FString::Printf( TEXT("%s.%s"), *Target.Name, TEXT("csv") );
+		FLocalizationConfigurationScript Script;
+
+		const FString ContentDirRelativeToGameDir = MakePathRelativeToProjectDirectory(FPaths::GameContentDir());
+
+		// GatherTextStep0 - GenerateTextLocalizationResource
+		{
+			FConfigSection& ConfigSection = Script.GatherTextStep(0);
+
+			// CommandletClass
+			ConfigSection.Add( TEXT("CommandletClass"), TEXT("GenerateTextLocalizationResource") );
+
+			const FString SourcePath = ContentDirRelativeToGameDir / TEXT("Localization") / Target.Name;
+			ConfigSection.Add( TEXT("SourcePath"), SourcePath );
+			const FString DestinationPath = ContentDirRelativeToGameDir / TEXT("Localization") / Target.Name;
+			ConfigSection.Add( TEXT("DestinationPath"), DestinationPath );
+
+			ConfigSection.Add( TEXT("ManifestName"), FString::Printf( TEXT("%s.%s"), *Target.Name, TEXT("manifest") ) );
+			ConfigSection.Add( TEXT("ResourceName"), FString::Printf( TEXT("%s.%s"), *Target.Name, TEXT("locres") ) );
+
+			TArray<const FCultureStatistics*> AllCultureStatistics;
+			AllCultureStatistics.Add(&Target.NativeCultureStatistics);
+			for (const FCultureStatistics& SupportedCultureStatistics : Target.SupportedCulturesStatistics)
+			{
+				AllCultureStatistics.Add(&SupportedCultureStatistics);
+			}
+
+			for (const FCultureStatistics* CultureStatistics : AllCultureStatistics)
+			{
+				ConfigSection.Add( TEXT("CulturesToGenerate"), CultureStatistics->CultureName );
+			}
+		}
+
+		Script.Dirty = true;
+
+		return Script;
 	}
 
-	FString GetConflictReportPath(const FLocalizationTargetSettings& Target)
+	FString GetCompileScriptPath(const FLocalizationTargetSettings& Target)
 	{
-		return GetDataDirectory(Target) / FString::Printf( TEXT("%s_Conflicts.%s"), *Target.Name, TEXT("txt") );
+		return GetScriptDirectory() / FString::Printf( TEXT("%s_Compile.%s"), *(Target.Name), TEXT("ini") );
 	}
 }
\ No newline at end of file
diff --git a/Engine/Source/Editor/LocalizationDashboard/Private/LocalizationConfigurationScript.h b/Engine/Source/Editor/LocalizationDashboard/Private/LocalizationConfigurationScript.h
--- a/Engine/Source/Editor/LocalizationDashboard/Private/LocalizationConfigurationScript.h
+++ b/Engine/Source/Editor/LocalizationDashboard/Private/LocalizationConfigurationScript.h
@@ -31,6 +31,8 @@ namespace LocalizationConfigurationScript
 	FString GetDefaultPOFileName(const FLocalizationTargetSettings& Target);
 	FString GetDefaultPOPath(const FLocalizationTargetSettings& Target, const FString& CultureName);
 	FString GetLocResPath(const FLocalizationTargetSettings& Target, const FString& CultureName);
+	FString GetWordCountCSVPath(const FLocalizationTargetSettings& Target);
+	FString GetConflictReportPath(const FLocalizationTargetSettings& Target);
 
 	FLocalizationConfigurationScript GenerateGatherScript(const FLocalizationTargetSettings& Target);
 	FString GetGatherScriptPath(const FLocalizationTargetSettings& Target);
@@ -43,6 +45,7 @@ namespace LocalizationConfigurationScript
 
 	FLocalizationConfigurationScript GenerateReportScript(const FLocalizationTargetSettings& Target);
 	FString GetReportScriptPath(const FLocalizationTargetSettings& Target);
-	FString GetWordCountCSVPath(const FLocalizationTargetSettings& Target);
-	FString GetConflictReportPath(const FLocalizationTargetSettings& Target);
+
+	FLocalizationConfigurationScript GenerateCompileScript(const FLocalizationTargetSettings& Target);
+	FString GetCompileScriptPath(const FLocalizationTargetSettings& Target);
 }
\ No newline at end of file
diff --git a/Engine/Source/Editor/LocalizationDashboard/Private/LocalizationDashboard.cpp b/Engine/Source/Editor/LocalizationDashboard/Private/LocalizationDashboard.cpp
--- a/Engine/Source/Editor/LocalizationDashboard/Private/LocalizationDashboard.cpp
+++ b/Engine/Source/Editor/LocalizationDashboard/Private/LocalizationDashboard.cpp
@@ -51,6 +51,7 @@ namespace
 		void ImportAllTargets();
 		void ExportAllTargets();
 		void UpdateTargetFromReports(ULocalizationTarget* const LocalizationTarget);
+		void CompileAllTargets();
 
 		TSharedRef<ITableRow> OnGenerateRow(TSharedPtr<IPropertyHandle> TargetObjectPropertyHandle, const TSharedRef<STableViewBase>& Table);
 		FReply OnNewTargetButtonClicked();
@@ -123,6 +124,7 @@ namespace
 		TSharedPtr<FUICommandInfo> GatherAllTargets;
 		TSharedPtr<FUICommandInfo> ImportAllTargets;
 		TSharedPtr<FUICommandInfo> ExportAllTargets;
+		TSharedPtr<FUICommandInfo> CompileAllTargets;
 
 		/** Initialize commands */
 		virtual void RegisterCommands() override;
@@ -133,6 +135,7 @@ namespace
 		UI_COMMAND( GatherAllTargets, "Gather All", "Gathers text for all targets in the project.", EUserInterfaceActionType::Button, FInputGesture() );
 		UI_COMMAND( ImportAllTargets, "Import All", "Imports translations for all cultures of all targets in the project.", EUserInterfaceActionType::Button, FInputGesture() );
 		UI_COMMAND( ExportAllTargets, "Export All", "Exports translations for all cultures of all targets in the project.", EUserInterfaceActionType::Button, FInputGesture() );
+		UI_COMMAND( CompileAllTargets, "Compile All", "Compiles translations for all targets in the project.", EUserInterfaceActionType::Button, FInputGesture() );
 	}
 
 	void FProjectLocalizationSettingsDetailsCustomization::CustomizeDetails(IDetailLayoutBuilder& DetailBuilder)
@@ -198,6 +201,9 @@ namespace
 				CommandList->MapAction( FLocalizationDashboardCommands::Get().ExportAllTargets, FExecuteAction::CreateSP(this, &FProjectLocalizationSettingsDetailsCustomization::ExportAllTargets));
 				ToolBarBuilder.AddToolBarButton(FLocalizationDashboardCommands::Get().ExportAllTargets, NAME_None, TAttribute<FText>(), TAttribute<FText>(), FSlateIcon(FEditorStyle::GetStyleSetName(), "LocalizationDashboard.ExportForAllTargetsCultures"));
 
+				CommandList->MapAction( FLocalizationDashboardCommands::Get().CompileAllTargets, FExecuteAction::CreateSP(this, &FProjectLocalizationSettingsDetailsCustomization::CompileAllTargets));
+				ToolBarBuilder.AddToolBarButton(FLocalizationDashboardCommands::Get().CompileAllTargets, NAME_None, TAttribute<FText>(), TAttribute<FText>(), FSlateIcon(FEditorStyle::GetStyleSetName(), "LocalizationDashboard.CompileAllTargets"));
+
 				BuildTargetsList();
 
 				DetailCategoryBuilder.AddCustomRow( LOCTEXT("TargetsFilterString", "Targets") )
@@ -496,6 +502,18 @@ namespace
 		//}
 	}
 
+	void FProjectLocalizationSettingsDetailsCustomization::CompileAllTargets()
+	{
+		// Execute compile.
+		const TSharedPtr<SWindow> ParentWindow = FSlateApplication::Get().FindWidgetWindow(DetailLayoutBuilder->GetDetailsView().AsShared());
+		TArray<FLocalizationTargetSettings*> TargetSettings;
+		for (ULocalizationTarget* const TargetObject : Settings->TargetObjects)
+		{
+			TargetSettings.Add(&TargetObject->Settings);
+		}
+		LocalizationCommandletTasks::CompileTargets(ParentWindow.ToSharedRef(), TargetSettings);
+	}
+
 	TSharedRef<ITableRow> FProjectLocalizationSettingsDetailsCustomization::OnGenerateRow(TSharedPtr<IPropertyHandle> TargetObjectPropertyHandle, const TSharedRef<STableViewBase>& Table)
 	{
 		return SNew(SLocalizationDashboardTargetRow, Table, DetailLayoutBuilder->GetPropertyUtilities(), TargetObjectPropertyHandle.ToSharedRef());
diff --git a/Engine/Source/Editor/LocalizationDashboard/Private/SLocalizationDashboardTargetRow.cpp b/Engine/Source/Editor/LocalizationDashboard/Private/SLocalizationDashboardTargetRow.cpp
--- a/Engine/Source/Editor/LocalizationDashboard/Private/SLocalizationDashboardTargetRow.cpp
+++ b/Engine/Source/Editor/LocalizationDashboard/Private/SLocalizationDashboardTargetRow.cpp
@@ -124,6 +124,23 @@ TSharedRef<SWidget> SLocalizationDashboardTargetRow::GenerateWidgetForColumn( co
 				]
 			];
 
+		// Compile
+		HorizontalBox->AddSlot()
+			.FillWidth(1.0f)
+			.HAlign(HAlign_Center)
+			.VAlign(VAlign_Center)
+			[
+				SNew(SButton)
+				.ButtonStyle( FEditorStyle::Get(), TEXT("HoverHintOnly") )
+				.ToolTipText( LOCTEXT("CompileButtonLabel", "Compile") )
+				.OnClicked(this, &SLocalizationDashboardTargetRow::Compile)
+				.Content()
+				[
+					SNew(SImage)
+					.Image( FEditorStyle::GetBrush("LocalizationDashboard.CompileTarget") )
+				]
+			];
+
 		// Delete Target
 		HorizontalBox->AddSlot()
 			.FillWidth(1.0f)
@@ -406,6 +423,19 @@ FReply SLocalizationDashboardTargetRow::ExportAll()
 	return FReply::Handled();
 }
 
+FReply SLocalizationDashboardTargetRow::Compile()
+{
+	ULocalizationTarget* const LocalizationTarget = GetTarget();
+	if (LocalizationTarget)
+	{
+		// Execute compile.
+		const TSharedPtr<SWindow> ParentWindow = FSlateApplication::Get().FindWidgetWindow(AsShared());
+		LocalizationCommandletTasks::CompileTarget(ParentWindow.ToSharedRef(), LocalizationTarget->Settings);
+	}
+
+	return FReply::Handled();
+}
+
 FReply SLocalizationDashboardTargetRow::EnqueueDeletion()
 {
 	PropertyUtilities->EnqueueDeferredAction(FSimpleDelegate::CreateSP(this, &SLocalizationDashboardTargetRow::Delete));
diff --git a/Engine/Source/Editor/LocalizationDashboard/Private/SLocalizationDashboardTargetRow.h b/Engine/Source/Editor/LocalizationDashboard/Private/SLocalizationDashboardTargetRow.h
--- a/Engine/Source/Editor/LocalizationDashboard/Private/SLocalizationDashboardTargetRow.h
+++ b/Engine/Source/Editor/LocalizationDashboard/Private/SLocalizationDashboardTargetRow.h
@@ -34,6 +34,8 @@ private:
 	FReply Gather();
 	FReply ImportAll();
 	FReply ExportAll();
+	FReply Compile();
+
 	FReply EnqueueDeletion();
 	void Delete();
 
diff --git a/Engine/Source/Editor/LocalizationDashboard/Private/SLocalizationTargetEditor.cpp b/Engine/Source/Editor/LocalizationDashboard/Private/SLocalizationTargetEditor.cpp
--- a/Engine/Source/Editor/LocalizationDashboard/Private/SLocalizationTargetEditor.cpp
+++ b/Engine/Source/Editor/LocalizationDashboard/Private/SLocalizationTargetEditor.cpp
@@ -51,6 +51,7 @@ namespace
 		void Gather();
 		void ImportAllCultures();
 		void ExportAllCultures();
+		void Compile();
 		void RefreshWordCounts();
 		void UpdateTargetFromReports();
 
@@ -112,6 +113,7 @@ namespace
 		TSharedPtr<FUICommandInfo> ImportAllCultures;
 		TSharedPtr<FUICommandInfo> ExportAllCultures;
 		TSharedPtr<FUICommandInfo> RefreshWordCounts;
+		TSharedPtr<FUICommandInfo> Compile;
 
 		/** Initialize commands */
 		virtual void RegisterCommands() override;
@@ -123,6 +125,7 @@ namespace
 		UI_COMMAND( ImportAllCultures, "Import All", "Imports translations for all cultures of this target.", EUserInterfaceActionType::Button, FInputGesture() );
 		UI_COMMAND( ExportAllCultures, "Export All", "Exports translations for all cultures of this target.", EUserInterfaceActionType::Button, FInputGesture() );
 		UI_COMMAND( RefreshWordCounts, "Count Words", "Recalculates the word counts for all cultures of this target.", EUserInterfaceActionType::Button, FInputGesture() );
+		UI_COMMAND( Compile, "Compile", "Compiles translations for this target.", EUserInterfaceActionType::Button, FInputGesture() );
 	}
 
 	void FLocalizationTargetDetailCustomization::CustomizeDetails(IDetailLayoutBuilder& DetailBuilder)
@@ -311,6 +314,10 @@ namespace
 				CommandList->MapAction( FLocalizationTargetEditorCommands::Get().RefreshWordCounts, FExecuteAction::CreateSP(this, &FLocalizationTargetDetailCustomization::RefreshWordCounts));
 				ToolBarBuilder.AddToolBarButton(FLocalizationTargetEditorCommands::Get().RefreshWordCounts, NAME_None, TAttribute<FText>(), TAttribute<FText>(), FSlateIcon(FEditorStyle::GetStyleSetName(), "LocalizationTargetEditor.RefreshWordCounts"));
 
+				CommandList->MapAction( FLocalizationTargetEditorCommands::Get().Compile, FExecuteAction::CreateSP(this, &FLocalizationTargetDetailCustomization::Compile));
+				ToolBarBuilder.AddToolBarButton(FLocalizationTargetEditorCommands::Get().Compile, NAME_None, TAttribute<FText>(), TAttribute<FText>(), FSlateIcon(FEditorStyle::GetStyleSetName(), "LocalizationTargetEditor.Compile"));
+
+
 				BuildListedCulturesList();
 
 				DetailCategoryBuilder.AddCustomRow( LOCTEXT("SupportedCultureNamesFilterString", "Supported Culture Names") )
@@ -732,6 +739,16 @@ namespace
 		}
 	}
 
+	void FLocalizationTargetDetailCustomization::Compile()
+	{
+		if (LocalizationTarget)
+		{
+			// Execute compile.
+			const TSharedPtr<SWindow> ParentWindow = FSlateApplication::Get().FindWidgetWindow(DetailLayoutBuilder->GetDetailsView().AsShared());
+			LocalizationCommandletTasks::CompileTarget(ParentWindow.ToSharedRef(), LocalizationTarget->Settings);
+		}
+	}
+
 	void FLocalizationTargetDetailCustomization::RefreshWordCounts()
 	{
 		if (LocalizationTarget)
