Fixed non-unity build.

diff --git a/Engine/Source/Editor/LocalizationDashboard/Private/LocalizationDashboardPrivatePCH.h b/Engine/Source/Editor/LocalizationDashboard/Private/LocalizationDashboardPrivatePCH.h
--- a/Engine/Source/Editor/LocalizationDashboard/Private/LocalizationDashboardPrivatePCH.h
+++ b/Engine/Source/Editor/LocalizationDashboard/Private/LocalizationDashboardPrivatePCH.h
@@ -2,7 +2,10 @@
 
 #pragma once
 
+
 /* Dependencies
  *****************************************************************************/
 
-#include "SlateBasics.h"
\ No newline at end of file
+#include "IDesktopPlatform.h"
+#include "DesktopPlatformModule.h"
+#include "SlateBasics.h"
