Linux CIS fix.

diff --git a/Engine/Source/Editor/LocalizationDashboard/Private/LocalizationDashboardPrivatePCH.h b/Engine/Source/Editor/LocalizationDashboard/Private/LocalizationDashboardPrivatePCH.h
--- a/Engine/Source/Editor/LocalizationDashboard/Private/LocalizationDashboardPrivatePCH.h
+++ b/Engine/Source/Editor/LocalizationDashboard/Private/LocalizationDashboardPrivatePCH.h
@@ -6,6 +6,7 @@
 /* Dependencies
  *****************************************************************************/
 
+#include "CoreUObject.h"
 #include "IDesktopPlatform.h"
 #include "DesktopPlatformModule.h"
 #include "SlateBasics.h"
