Fixed crashes in localization dashboard when executing localization commandlets due to threading and resource lifetime issues.

diff --git a/Engine/Source/Editor/LocalizationDashboard/Private/LocalizationCommandletExecution.cpp b/Engine/Source/Editor/LocalizationDashboard/Private/LocalizationCommandletExecution.cpp
--- a/Engine/Source/Editor/LocalizationDashboard/Private/LocalizationCommandletExecution.cpp
+++ b/Engine/Source/Editor/LocalizationDashboard/Private/LocalizationCommandletExecution.cpp
@@ -1,4 +1,4 @@
-// Copyright 1998-2014 Epic Games, Inc. All Rights Reserved.
+﻿// Copyright 1998-2014 Epic Games, Inc. All Rights Reserved.
 
 #include "LocalizationDashboardPrivatePCH.h"
 #include "LocalizationCommandletExecution.h"
@@ -85,7 +85,7 @@ namespace
 		{
 			FCriticalSection CriticalSection;
 			FString String;
-		} PendingData;
+		} PendingLogData;
 
 
 		TSharedPtr<SWindow> ParentWindow;
@@ -256,15 +256,15 @@ namespace
 		SCompoundWidget::Tick(AllottedGeometry, InCurrentTime, InDeltaTime);
 
 		// Poll for log output data.
-		if (!PendingData.String.IsEmpty())
+		if (!PendingLogData.String.IsEmpty())
 		{
 			FString String;
 
 			// Copy the pending data string to the local string
 			{
-				FScopeLock ScopeLock(&PendingData.CriticalSection);
-				String = PendingData.String;
-				PendingData.String.Empty();
+				FScopeLock ScopeLock(&PendingLogData.CriticalSection);
+				String = PendingLogData.String;
+				PendingLogData.String.Empty();
 			}
 
 			// Forward string to proper log.
@@ -291,8 +291,8 @@ namespace
 
 	void SLocalizationCommandletExecutor::Log(const FString& String)
 	{
-		FScopeLock ScopeLock(&PendingData.CriticalSection);
-		PendingData.String += String;
+		FScopeLock ScopeLock(&PendingLogData.CriticalSection);
+		PendingLogData.String += String;
 	}
 
 	void SLocalizationCommandletExecutor::OnCommandletProcessCompletion(const int32 ReturnCode)
@@ -338,10 +338,10 @@ namespace
 		class FCommandletLogPump : public FRunnable
 		{
 		public:
-			FCommandletLogPump(void* const InReadPipe, const FProcHandle& InCommandletProcessHandle, const TSharedRef<SLocalizationCommandletExecutor>& InCommandletWidget)
+			FCommandletLogPump(void* const InReadPipe, const FProcHandle& InCommandletProcessHandle, SLocalizationCommandletExecutor& InCommandletWidget)
 				: ReadPipe(InReadPipe)
 				, CommandletProcessHandle(InCommandletProcessHandle)
-				, CommandletWidget(InCommandletWidget)
+				, CommandletWidget(&InCommandletWidget)
 			{
 			}
 
@@ -356,10 +356,9 @@ namespace
 					if (!PipeString.IsEmpty())
 					{
 						// Add strings to log.
-						const TSharedPtr<SLocalizationCommandletExecutor> Widget = CommandletWidget.Pin();
-						if (Widget.IsValid())
+						if (CommandletWidget)
 						{
-							Widget->Log(PipeString);
+							CommandletWidget->Log(PipeString);
 						}
 					}
 
@@ -380,11 +379,11 @@ namespace
 		private:
 			void* const ReadPipe;
 			FProcHandle CommandletProcessHandle;
-			const TWeakPtr<SLocalizationCommandletExecutor> CommandletWidget;
+			SLocalizationCommandletExecutor* const CommandletWidget;
 		};
 
 		// Launch runnable thread.
-		Runnable = new FCommandletLogPump(CommandletProcess->GetReadPipe(), CommandletProcess->GetHandle(), SharedThis(this));
+		Runnable = new FCommandletLogPump(CommandletProcess->GetReadPipe(), CommandletProcess->GetHandle(), *this);
 		RunnableThread = FRunnableThread::Create(Runnable, TEXT("Localization Commandlet Log Pump Thread"));
 	}
 
@@ -596,7 +595,7 @@ namespace
 		{
 			void* ParentWindowWindowHandle = NULL;
 
-			const TSharedPtr<SWindow>& ParentWindow = FSlateApplication::Get().FindWidgetWindow(SharedThis(this));
+			const TSharedPtr<SWindow>& ParentWindow = FSlateApplication::Get().FindWidgetWindow(AsShared());
 			if (ParentWindow.IsValid() && ParentWindow->GetNativeWindow().IsValid())
 			{
 				ParentWindowWindowHandle = ParentWindow->GetNativeWindow()->GetOSWindowHandle();
@@ -649,7 +648,6 @@ bool LocalizationCommandletExecution::Execute(const TSharedRef<SWindow>& ParentW
 	CommandletWindow->SetContent(CommandletExecutor);
 
 	FSlateApplication::Get().AddModalWindow(CommandletWindow, ParentWindow, false);
-
 	return CommandletExecutor->WasSuccessful();
 }
 
@@ -672,7 +670,7 @@ TSharedPtr<FLocalizationCommandletProcess> FLocalizationCommandletProcess::Execu
 	ISourceControlModule& SourceControlModule = ISourceControlModule::Get();
 	if (SourceControlModule.IsEnabled())
 	{
-		if (!SourceControlModule.GetProvider().IsAvailable())
+		if (SourceControlModule.GetProvider().IsAvailable())
 		{
 			CommandletArguments = FString::Printf( TEXT("%s -EnableSCC -DisableSCCSubmit"), *CommandletArguments );
 		}
@@ -692,7 +690,7 @@ TSharedPtr<FLocalizationCommandletProcess> FLocalizationCommandletProcess::Execu
 
 FLocalizationCommandletProcess::~FLocalizationCommandletProcess()
 {
-	if (FPlatformProcess::IsProcRunning(ProcessHandle))
+	if (ProcessHandle.IsValid() && FPlatformProcess::IsProcRunning(ProcessHandle))
 	{
 		FPlatformProcess::TerminateProc(ProcessHandle);
 	}
