Added menu to localization dashboard.

diff --git a/Engine/Source/Editor/LocalizationDashboard/LocalizationDashboard.Build.cs b/Engine/Source/Editor/LocalizationDashboard/LocalizationDashboard.Build.cs
--- a/Engine/Source/Editor/LocalizationDashboard/LocalizationDashboard.Build.cs
+++ b/Engine/Source/Editor/LocalizationDashboard/LocalizationDashboard.Build.cs
@@ -25,7 +25,8 @@ public class LocalizationDashboard : ModuleRules
                 "UnrealEd",
                 "EditorStyle",
 				"DesktopPlatform",
-                "TranslationEditor"
+                "TranslationEditor",
+                "MainFrame"
 			}
 		);
 
diff --git a/Engine/Source/Editor/LocalizationDashboard/Private/LocalizationDashboard.cpp b/Engine/Source/Editor/LocalizationDashboard/Private/LocalizationDashboard.cpp
--- a/Engine/Source/Editor/LocalizationDashboard/Private/LocalizationDashboard.cpp
+++ b/Engine/Source/Editor/LocalizationDashboard/Private/LocalizationDashboard.cpp
@@ -1,4 +1,4 @@
-// Copyright 1998-2014 Epic Games, Inc. All Rights Reserved.
+﻿// Copyright 1998-2014 Epic Games, Inc. All Rights Reserved.
 
 #include "LocalizationDashboardPrivatePCH.h"
 #include "LocalizationDashboard.h"
@@ -22,6 +22,7 @@
 #include "LocalizationCommandletTasks.h"
 #include "FileHelpers.h"
 #include "SLocalizationTargetEditor.h"
+#include "IMainFrameModule.h"
 
 #define LOCTEXT_NAMESPACE "LocalizationDashboard"
 
@@ -549,14 +550,14 @@ public:
 	TWeakPtr<SDockTab> ShowTargetEditor(ULocalizationTarget* const LocalizationTarget);
 
 private:
-	static const FName DetailsTabName;
+	static const FName ProjectTargetSetDetailsTabName;
 	static const FName DocumentsTabName;
 
 	TSharedPtr<FTabManager> TabManager;
 	TMap< TWeakObjectPtr<ULocalizationTarget>, TWeakPtr<SDockTab> > TargetToTabMap;
 };
 
-const FName SLocalizationDashboard::DetailsTabName("Details");
+const FName SLocalizationDashboard::ProjectTargetSetDetailsTabName("ProjectTargets");
 const FName SLocalizationDashboard::DocumentsTabName("Documents");
 
 void SLocalizationDashboard::Construct(const FArguments& InArgs, const TSharedPtr<SWindow>& OwningWindow, const TSharedRef<SDockTab>& OwningTab, UProjectLocalizationSettings* const Settings)
@@ -573,8 +574,8 @@ void SLocalizationDashboard::Construct(const FArguments& InArgs, const TSharedPt
 
 	const auto& CreateDetailsTab = [Settings](const FSpawnTabArgs& SpawnTabArgs) -> TSharedRef<SDockTab>
 	{
-		const TSharedRef<SDockTab> DockTab = SNew(SDockTab);
-			//.OnCanCloseTab_Lambda([]()->bool{return false;});
+		const TSharedRef<SDockTab> DockTab = SNew(SDockTab)
+			.Label(LOCTEXT("ProjectTargetsTabLabel", "Project Targets"));
 
 		FPropertyEditorModule& PropertyModule = FModuleManager::LoadModuleChecked<FPropertyEditorModule>("PropertyEditor");
 		FDetailsViewArgs DetailsViewArgs(false, false, false, FDetailsViewArgs::ENameAreaSettings::HideNameArea, false, nullptr, false, NAME_None);
@@ -587,9 +588,12 @@ void SLocalizationDashboard::Construct(const FArguments& InArgs, const TSharedPt
 
 		return DockTab;
 	};
-	TabManager->RegisterTabSpawner(DetailsTabName, FOnSpawnTab::CreateLambda(CreateDetailsTab));
+	const TSharedRef<FWorkspaceItem> TargetSetsWorkspaceMenuCategory = TabManager->AddLocalWorkspaceMenuCategory(LOCTEXT("TargetSetsWorkspaceMenuCategory", "Target Sets"));
+	FTabSpawnerEntry& ProjectTargetSetTabSpawnerEntry = TabManager->RegisterTabSpawner(ProjectTargetSetDetailsTabName, FOnSpawnTab::CreateLambda(CreateDetailsTab))
+		.SetDisplayName(LOCTEXT("ProjectTargetsDetailTabSpawner", "Project Targets"));
+	TargetSetsWorkspaceMenuCategory->AddItem(ProjectTargetSetTabSpawnerEntry.AsShared());
 
-	TSharedRef<FTabManager::FLayout> Layout = FTabManager::NewLayout("LocalizationDashboard_Experimental_V1")
+	TSharedRef<FTabManager::FLayout> Layout = FTabManager::NewLayout("LocalizationDashboard_Experimental_V3")
 		->AddArea
 		(
 		FTabManager::NewPrimaryArea()
@@ -598,7 +602,7 @@ void SLocalizationDashboard::Construct(const FArguments& InArgs, const TSharedPt
 		(
 		FTabManager::NewStack()
 		->SetHideTabWell(true)
-		->AddTab(DetailsTabName, ETabState::OpenedTab)
+		->AddTab(ProjectTargetSetDetailsTabName, ETabState::OpenedTab)
 		)
 		->Split
 		(
@@ -610,9 +614,23 @@ void SLocalizationDashboard::Construct(const FArguments& InArgs, const TSharedPt
 
 	Layout = FLayoutSaveRestore::LoadFromConfig(GEditorLayoutIni, Layout);
 
+	const TSharedRef<FExtender> MenuExtender = MakeShareable(new FExtender());
+	IMainFrameModule& MainFrameModule = FModuleManager::LoadModuleChecked<IMainFrameModule>( "MainFrame" );
+	const TSharedRef<SWidget> MenuWidget = MainFrameModule.MakeMainMenu( TabManager, MenuExtender );
+
 	ChildSlot
 		[
+			SNew(SVerticalBox)
+			+SVerticalBox::Slot()
+			.AutoHeight()
+			[
+				MenuWidget
+			]
+			+SVerticalBox::Slot()
+			.FillHeight(1.0f)
+			[
 			TabManager->RestoreFrom(Layout, OwningWindow).ToSharedRef()
+			]
 		];
 }
 
@@ -709,8 +727,7 @@ void FLocalizationDashboard::RegisterTabSpawner()
 	};
 
 	FGlobalTabmanager::Get()->RegisterTabSpawner(TabName, FOnSpawnTab::CreateLambda( TFunction<TSharedRef<SDockTab> (const FSpawnTabArgs&)>(SpawnMainTab) ) )
-		.SetDisplayName(LOCTEXT("MainTabTitle", "Localization Dashboard"))
-		.SetMenuType(ETabSpawnerMenuType::Hidden);
+		.SetDisplayName(LOCTEXT("MainTabTitle", "Localization Dashboard"));
 }
 
 void FLocalizationDashboard::UnregisterTabSpawner()
diff --git a/Engine/Source/Editor/MainFrame/MainFrame.Build.cs b/Engine/Source/Editor/MainFrame/MainFrame.Build.cs
--- a/Engine/Source/Editor/MainFrame/MainFrame.Build.cs
+++ b/Engine/Source/Editor/MainFrame/MainFrame.Build.cs
@@ -39,7 +39,6 @@ public class MainFrame : ModuleRules
 				"Projects",
                 "DeviceProfileEditor",
                 "UndoHistory",
-                "LocalizationDashboard",
 			}
 		);
 
