Renamed project localization settings to localization target set, as its purpose is simpler now. Done in preparation for having UE4 targets versus project-specific targets.

diff --git a/Engine/Source/Editor/LocalizationDashboard/Classes/LocalizationTarget.h b/Engine/Source/Editor/LocalizationDashboard/Classes/LocalizationTarget.h
--- a/Engine/Source/Editor/LocalizationDashboard/Classes/LocalizationTarget.h
+++ b/Engine/Source/Editor/LocalizationDashboard/Classes/LocalizationTarget.h
@@ -62,10 +62,10 @@ struct FLocalizationTargetSettings
 	bool DeleteFiles(const FString* const Culture = nullptr) const;
 };
 
-UCLASS(Within=ProjectLocalizationSettings)
+UCLASS(Within=LocalizationTargetSet)
 class ULocalizationTarget : public UObject
 {
-	GENERATED_UCLASS_BODY()
+	GENERATED_BODY()
 
 public:
 	UPROPERTY(EditAnywhere, Category = "Target")
@@ -75,4 +75,24 @@ public:
 #if WITH_EDITOR
 	virtual void PostEditChangeProperty(FPropertyChangedEvent& PropertyChangedEvent) override;
 #endif
+};
+
+UCLASS(Config=Game, defaultconfig)
+class ULocalizationTargetSet : public UObject
+{
+	GENERATED_BODY()
+
+public:
+	UPROPERTY(EditAnywhere, Category = "Targets")
+	TArray<ULocalizationTarget*> TargetObjects;
+
+private:
+	UPROPERTY(config)
+	TArray<FLocalizationTargetSettings> Targets;
+
+public:
+#if WITH_EDITOR
+	virtual void PostInitProperties() override;
+	virtual void PostEditChangeProperty(FPropertyChangedEvent& PropertyChangedEvent) override;
+#endif
 };
\ No newline at end of file
diff --git a/Engine/Source/Editor/LocalizationDashboard/Classes/ProjectLocalizationSettings.h b/Engine/Source/Editor/LocalizationDashboard/Classes/ProjectLocalizationSettings.h
deleted file mode 100644
diff --git a/Engine/Source/Editor/LocalizationDashboard/Private/LocalizationConfigurationScript.cpp b/Engine/Source/Editor/LocalizationDashboard/Private/LocalizationConfigurationScript.cpp
--- a/Engine/Source/Editor/LocalizationDashboard/Private/LocalizationConfigurationScript.cpp
+++ b/Engine/Source/Editor/LocalizationDashboard/Private/LocalizationConfigurationScript.cpp
@@ -2,7 +2,7 @@
 
 #include "LocalizationDashboardPrivatePCH.h"
 #include "LocalizationConfigurationScript.h"
-#include "Classes/ProjectLocalizationSettings.h"
+#include "Classes/LocalizationTarget.h"
 
 namespace LocalizationConfigurationScript
 {
@@ -82,10 +82,10 @@ namespace LocalizationConfigurationScript
 		{
 			FConfigSection& ConfigSection = Script.CommonSettings();
 
-			const UProjectLocalizationSettings* const ProjectLocalizationSettings = GetDefault<UProjectLocalizationSettings>(UProjectLocalizationSettings::StaticClass());
+			const ULocalizationTargetSet* const LocalizationTargetSet = GetDefault<ULocalizationTargetSet>(ULocalizationTargetSet::StaticClass());
 			for (const FString& TargetDependencyName : Target.TargetDependencies)
 			{
-				ULocalizationTarget* const * OtherTarget = ProjectLocalizationSettings->TargetObjects.FindByPredicate([&TargetDependencyName](ULocalizationTarget* const OtherTarget)->bool{return OtherTarget->Settings.Name == TargetDependencyName;});
+				ULocalizationTarget* const * OtherTarget = LocalizationTargetSet->TargetObjects.FindByPredicate([&TargetDependencyName](ULocalizationTarget* const OtherTarget)->bool{return OtherTarget->Settings.Name == TargetDependencyName;});
 				if (OtherTarget)
 				{
 					ConfigSection.Add( TEXT("ManifestDependencies"), MakePathRelativeToProjectDirectory(GetManifestPath((*OtherTarget)->Settings)) );
diff --git a/Engine/Source/Editor/LocalizationDashboard/Private/LocalizationDashboard.cpp b/Engine/Source/Editor/LocalizationDashboard/Private/LocalizationDashboard.cpp
--- a/Engine/Source/Editor/LocalizationDashboard/Private/LocalizationDashboard.cpp
+++ b/Engine/Source/Editor/LocalizationDashboard/Private/LocalizationDashboard.cpp
@@ -5,7 +5,7 @@
 #include "IDetailCustomization.h"
 #include "EditorStyleSet.h"
 #include "DetailLayoutBuilder.h"
-#include "ProjectLocalizationSettings.h"
+#include "LocalizationTarget.h"
 #include "PropertyEditorModule.h"
 #include "IDetailsView.h"
 #include "DetailWidgetRow.h"
@@ -34,10 +34,10 @@ namespace
 {
 	struct FLocalizationServiceProviderWrapper;
 
-	class FProjectLocalizationSettingsDetailsCustomization : public IDetailCustomization
+	class FLocalizationTargetSetDetailsCustomization : public IDetailCustomization
 	{
 	public:
-		FProjectLocalizationSettingsDetailsCustomization(UProjectLocalizationSettings* const InSettings);
+		FLocalizationTargetSetDetailsCustomization(ULocalizationTargetSet* const InTargetSet);
 		void CustomizeDetails(IDetailLayoutBuilder& DetailBuilder) override;
 
 	private:
@@ -58,7 +58,7 @@ namespace
 		FReply OnNewTargetButtonClicked();
 
 	private:
-		UProjectLocalizationSettings* Settings;
+		ULocalizationTargetSet* TargetSet;
 
 		IDetailLayoutBuilder* DetailLayoutBuilder;
 
@@ -82,8 +82,8 @@ namespace
 		ILocalizationServiceProvider* Provider;
 	};
 
-	FProjectLocalizationSettingsDetailsCustomization::FProjectLocalizationSettingsDetailsCustomization(UProjectLocalizationSettings* const InSettings)
-		: Settings(InSettings)
+	FLocalizationTargetSetDetailsCustomization::FLocalizationTargetSetDetailsCustomization(ULocalizationTargetSet* const InTargetSet)
+		: TargetSet(InTargetSet)
 		, DetailLayoutBuilder(nullptr)
 		, ServiceProviderCategoryBuilder(nullptr)
 		, NewEntryIndexToBeInitialized(INDEX_NONE)
@@ -139,70 +139,33 @@ namespace
 		UI_COMMAND( CompileAllTargets, "Compile All", "Compiles translations for all targets in the project.", EUserInterfaceActionType::Button, FInputGesture() );
 	}
 
-	void FProjectLocalizationSettingsDetailsCustomization::CustomizeDetails(IDetailLayoutBuilder& DetailBuilder)
+	void FLocalizationTargetSetDetailsCustomization::CustomizeDetails(IDetailLayoutBuilder& DetailBuilder)
 	{
 		DetailLayoutBuilder = &DetailBuilder;
 
-		{
-			ServiceProviderPropertyHandle = DetailBuilder.GetProperty(GET_MEMBER_NAME_CHECKED(UProjectLocalizationSettings, LocalizationServiceProvider));
-
-			if (ServiceProviderPropertyHandle.IsValid() && ServiceProviderPropertyHandle->IsValidHandle())
-			{
-				ServiceProviderPropertyHandle->MarkHiddenByCustomization();
-				//ServiceProviderCategoryBuilder = &( DetailBuilder.EditCategory( FName("LocalizationServiceProvider") ) );
-
-				//FDetailWidgetRow& ServiceProviderDetailWidgetRow = ServiceProviderCategoryBuilder->AddCustomRow( LOCTEXT("LocalizationServiceProviderFilterString", "Localization Service Provider") );
-				//ServiceProviderDetailWidgetRow.NameContent()
-				//	[
-				//		ServiceProviderPropertyHandle->CreatePropertyNameWidget()
-				//	];
-
-				//ILocalizationServiceProvider* const LSP = ILocalizationDashboardModule::Get().GetCurrentLocalizationServiceProvider();
-				//TSharedPtr<FLocalizationServiceProviderWrapper>* LSPWrapper = Providers.FindByPredicate(FindProviderPredicate(LSP));
-				//ServiceProviderDetailWidgetRow.ValueContent()
-				//	[
-				//		SNew(SComboBox<TSharedPtr<FLocalizationServiceProviderWrapper>>)
-				//		.InitiallySelectedItem(LSPWrapper ? *LSPWrapper : nullptr)
-				//		.OptionsSource(&Providers)
-				//		.OnGenerateWidget(this, &FProjectLocalizationSettingsDetailsCustomization::ServiceProviderComboBox_OnGenerateWidget)
-				//		.OnSelectionChanged(this, &FProjectLocalizationSettingsDetailsCustomization::ServiceProviderComboBox_OnSelectionChanged)
-				//		.Content()
-				//		[
-				//			SNew(STextBlock)
-				//			.Text(this, &FProjectLocalizationSettingsDetailsCustomization::GetCurrentServiceProviderDisplayName)
-				//		]
-				//	];
-
-				//if (LSP)
-				//{
-				//	LSP->CustomizeSettingsDetails(*ServiceProviderCategoryBuilder);
-				//}
-			}
-		}
-
 		{
 			IDetailCategoryBuilder& DetailCategoryBuilder = DetailBuilder.EditCategory(FName("Targets"));
-			TargetObjectsPropertyHandle = DetailBuilder.GetProperty(GET_MEMBER_NAME_CHECKED(UProjectLocalizationSettings,TargetObjects));
+			TargetObjectsPropertyHandle = DetailBuilder.GetProperty(GET_MEMBER_NAME_CHECKED(ULocalizationTargetSet,TargetObjects));
 			if (TargetObjectsPropertyHandle->IsValidHandle())
 			{
 				TargetObjectsPropertyHandle->MarkHiddenByCustomization();
-				TargetsArrayPropertyHandle_OnNumElementsChanged = FSimpleDelegate::CreateSP(this, &FProjectLocalizationSettingsDetailsCustomization::RebuildTargetsList);
+				TargetsArrayPropertyHandle_OnNumElementsChanged = FSimpleDelegate::CreateSP(this, &FLocalizationTargetSetDetailsCustomization::RebuildTargetsList);
 				TargetObjectsPropertyHandle->AsArray()->SetOnNumElementsChanged(TargetsArrayPropertyHandle_OnNumElementsChanged);
 
 				FLocalizationDashboardCommands::Register();
 				const TSharedRef< FUICommandList > CommandList = MakeShareable(new FUICommandList);
 				FToolBarBuilder ToolBarBuilder( CommandList, FMultiBoxCustomization::AllowCustomization("LocalizationDashboard") );
 
-				CommandList->MapAction( FLocalizationDashboardCommands::Get().GatherAllTargets, FExecuteAction::CreateSP(this, &FProjectLocalizationSettingsDetailsCustomization::GatherAllTargets));
+				CommandList->MapAction( FLocalizationDashboardCommands::Get().GatherAllTargets, FExecuteAction::CreateSP(this, &FLocalizationTargetSetDetailsCustomization::GatherAllTargets));
 				ToolBarBuilder.AddToolBarButton(FLocalizationDashboardCommands::Get().GatherAllTargets, NAME_None, TAttribute<FText>(), TAttribute<FText>(), FSlateIcon(FEditorStyle::GetStyleSetName(), "LocalizationDashboard.GatherAllTargets"));
 
-				CommandList->MapAction( FLocalizationDashboardCommands::Get().ImportAllTargets, FExecuteAction::CreateSP(this, &FProjectLocalizationSettingsDetailsCustomization::ImportAllTargets));
+				CommandList->MapAction( FLocalizationDashboardCommands::Get().ImportAllTargets, FExecuteAction::CreateSP(this, &FLocalizationTargetSetDetailsCustomization::ImportAllTargets));
 				ToolBarBuilder.AddToolBarButton(FLocalizationDashboardCommands::Get().ImportAllTargets, NAME_None, TAttribute<FText>(), TAttribute<FText>(), FSlateIcon(FEditorStyle::GetStyleSetName(), "LocalizationDashboard.ImportForAllTargetsCultures"));
 
-				CommandList->MapAction( FLocalizationDashboardCommands::Get().ExportAllTargets, FExecuteAction::CreateSP(this, &FProjectLocalizationSettingsDetailsCustomization::ExportAllTargets));
+				CommandList->MapAction( FLocalizationDashboardCommands::Get().ExportAllTargets, FExecuteAction::CreateSP(this, &FLocalizationTargetSetDetailsCustomization::ExportAllTargets));
 				ToolBarBuilder.AddToolBarButton(FLocalizationDashboardCommands::Get().ExportAllTargets, NAME_None, TAttribute<FText>(), TAttribute<FText>(), FSlateIcon(FEditorStyle::GetStyleSetName(), "LocalizationDashboard.ExportForAllTargetsCultures"));
 
-				CommandList->MapAction( FLocalizationDashboardCommands::Get().CompileAllTargets, FExecuteAction::CreateSP(this, &FProjectLocalizationSettingsDetailsCustomization::CompileAllTargets));
+				CommandList->MapAction( FLocalizationDashboardCommands::Get().CompileAllTargets, FExecuteAction::CreateSP(this, &FLocalizationTargetSetDetailsCustomization::CompileAllTargets));
 				ToolBarBuilder.AddToolBarButton(FLocalizationDashboardCommands::Get().CompileAllTargets, NAME_None, TAttribute<FText>(), TAttribute<FText>(), FSlateIcon(FEditorStyle::GetStyleSetName(), "LocalizationDashboard.CompileAllTargets"));
 
 				BuildTargetsList();
@@ -220,7 +183,7 @@ namespace
 							.AutoHeight()
 							[
 								SAssignNew(TargetsListView, SListView< TSharedPtr<IPropertyHandle> >)
-								.OnGenerateRow(this, &FProjectLocalizationSettingsDetailsCustomization::OnGenerateRow)
+								.OnGenerateRow(this, &FLocalizationTargetSetDetailsCustomization::OnGenerateRow)
 								.ListItemsSource(&TargetsList)
 								.SelectionMode(ESelectionMode::None)
 								.HeaderRow
@@ -259,14 +222,14 @@ namespace
 							[
 								SNew(SButton)
 								.Text(LOCTEXT("AddNewTargetButtonLabel", "Add New Target"))
-								.OnClicked(this, &FProjectLocalizationSettingsDetailsCustomization::OnNewTargetButtonClicked)
+								.OnClicked(this, &FLocalizationTargetSetDetailsCustomization::OnNewTargetButtonClicked)
 							]
 					];
 			}
 		}
 	}
 
-	void FProjectLocalizationSettingsDetailsCustomization::BuildTargetsList()
+	void FLocalizationTargetSetDetailsCustomization::BuildTargetsList()
 	{
 		const TSharedPtr<IPropertyHandleArray> TargetObjectsArrayPropertyHandle = TargetObjectsPropertyHandle->AsArray();
 		if (TargetObjectsArrayPropertyHandle.IsValid())
@@ -281,7 +244,7 @@ namespace
 		}
 	}
 
-	void FProjectLocalizationSettingsDetailsCustomization::RebuildTargetsList()
+	void FLocalizationTargetSetDetailsCustomization::RebuildTargetsList()
 	{
 		const TSharedPtr<IPropertyHandleArray> TargetObjectsArrayPropertyHandle = TargetObjectsPropertyHandle->AsArray();
 		if (TargetObjectsArrayPropertyHandle.IsValid() && NewEntryIndexToBeInitialized != INDEX_NONE)
@@ -289,7 +252,7 @@ namespace
 			const TSharedPtr<IPropertyHandle> TargetObjectPropertyHandle = TargetObjectsArrayPropertyHandle->GetElement(NewEntryIndexToBeInitialized);
 			if (TargetObjectPropertyHandle.IsValid() && TargetObjectPropertyHandle->IsValidHandle())
 			{
-				ULocalizationTarget* const NewTarget = NewObject<ULocalizationTarget>(Settings);
+				ULocalizationTarget* const NewTarget = NewObject<ULocalizationTarget>(TargetSet);
 
 				TArray<void*> RawData;
 				TargetObjectsPropertyHandle->AccessRawData(RawData);
@@ -326,20 +289,20 @@ namespace
 		}
 	}
 
-	FText FProjectLocalizationSettingsDetailsCustomization::GetCurrentServiceProviderDisplayName() const
+	FText FLocalizationTargetSetDetailsCustomization::GetCurrentServiceProviderDisplayName() const
 	{
 		ILocalizationServiceProvider* const LSP = ILocalizationDashboardModule::Get().GetCurrentLocalizationServiceProvider();
 		return LSP ? LSP->GetDisplayName() : LOCTEXT("NoServiceProviderName", "None");
 	}
 
-	TSharedRef<SWidget> FProjectLocalizationSettingsDetailsCustomization::ServiceProviderComboBox_OnGenerateWidget(TSharedPtr<FLocalizationServiceProviderWrapper> LSPWrapper) const
+	TSharedRef<SWidget> FLocalizationTargetSetDetailsCustomization::ServiceProviderComboBox_OnGenerateWidget(TSharedPtr<FLocalizationServiceProviderWrapper> LSPWrapper) const
 	{
 		ILocalizationServiceProvider* const LSP = LSPWrapper->Provider;
 		return	SNew(STextBlock)
 			.Text(LSP ? LSP->GetDisplayName() : LOCTEXT("NoServiceProviderName", "None"));
 	}
 
-	void FProjectLocalizationSettingsDetailsCustomization::ServiceProviderComboBox_OnSelectionChanged(TSharedPtr<FLocalizationServiceProviderWrapper> LSPWrapper, ESelectInfo::Type SelectInfo)
+	void FLocalizationTargetSetDetailsCustomization::ServiceProviderComboBox_OnSelectionChanged(TSharedPtr<FLocalizationServiceProviderWrapper> LSPWrapper, ESelectInfo::Type SelectInfo)
 	{
 		ILocalizationServiceProvider* const LSP = LSPWrapper->Provider;
 		FString ServiceProviderName = LSP ? LSP->GetName().ToString() : TEXT("None");
@@ -352,9 +315,9 @@ namespace
 		DetailLayoutBuilder->ForceRefreshDetails();
 	}
 
-	void FProjectLocalizationSettingsDetailsCustomization::GatherAllTargets()
+	void FLocalizationTargetSetDetailsCustomization::GatherAllTargets()
 	{
-		for (ULocalizationTarget* const LocalizationTarget : Settings->TargetObjects)
+		for (ULocalizationTarget* const LocalizationTarget : TargetSet->TargetObjects)
 		{
 			// Save unsaved packages.
 			const bool bPromptUserToSave = true;
@@ -386,19 +349,19 @@ namespace
 		// Execute gather.
 		const TSharedPtr<SWindow> ParentWindow = FSlateApplication::Get().FindWidgetWindow(DetailLayoutBuilder->GetDetailsView().AsShared());
 		TArray<FLocalizationTargetSettings*> TargetSettings;
-		for (ULocalizationTarget* const TargetObject : Settings->TargetObjects)
+		for (ULocalizationTarget* const TargetObject : TargetSet->TargetObjects)
 		{
 			TargetSettings.Add(&TargetObject->Settings);
 		}
 		LocalizationCommandletTasks::GatherTargets(ParentWindow.ToSharedRef(), TargetSettings);
 
-		for (ULocalizationTarget* const LocalizationTarget : Settings->TargetObjects)
+		for (ULocalizationTarget* const LocalizationTarget : TargetSet->TargetObjects)
 		{
 			UpdateTargetFromReports(LocalizationTarget);
 		}
 	}
 
-	void FProjectLocalizationSettingsDetailsCustomization::ImportAllTargets()
+	void FLocalizationTargetSetDetailsCustomization::ImportAllTargets()
 	{
 		IDesktopPlatform* DesktopPlatform = FDesktopPlatformModule::Get();
 		if (DesktopPlatform)
@@ -418,13 +381,13 @@ namespace
 			if (DesktopPlatform->OpenDirectoryDialog(ParentWindowWindowHandle, LOCTEXT("ImportAllTranslationsDialogTitle", "Import All Translations from Directory").ToString(), DefaultPath, OutputDirectory))
 			{
 				TArray<FLocalizationTargetSettings*> TargetSettings;
-				for (ULocalizationTarget* const TargetObject : Settings->TargetObjects)
+				for (ULocalizationTarget* const TargetObject : TargetSet->TargetObjects)
 				{
 					TargetSettings.Add(&TargetObject->Settings);
 				}
 				LocalizationCommandletTasks::ImportTargets(ParentWindow.ToSharedRef(), TargetSettings, TOptional<FString>(OutputDirectory));
 
-				for (ULocalizationTarget* const LocalizationTarget : Settings->TargetObjects)
+				for (ULocalizationTarget* const LocalizationTarget : TargetSet->TargetObjects)
 				{
 					UpdateTargetFromReports(LocalizationTarget);
 				}
@@ -432,7 +395,7 @@ namespace
 		}
 	}
 
-	void FProjectLocalizationSettingsDetailsCustomization::ExportAllTargets()
+	void FLocalizationTargetSetDetailsCustomization::ExportAllTargets()
 	{
 		IDesktopPlatform* DesktopPlatform = FDesktopPlatformModule::Get();
 		if (DesktopPlatform)
@@ -451,7 +414,7 @@ namespace
 			if (DesktopPlatform->OpenDirectoryDialog(ParentWindowWindowHandle, LOCTEXT("ExportAllTranslationsDialogTitle", "Export All Translations to Directory").ToString(), DefaultPath, OutputDirectory))
 			{
 				TArray<FLocalizationTargetSettings*> TargetSettings;
-				for (ULocalizationTarget* const TargetObject : Settings->TargetObjects)
+				for (ULocalizationTarget* const TargetObject : TargetSet->TargetObjects)
 				{
 					TargetSettings.Add(&TargetObject->Settings);
 				}
@@ -460,7 +423,7 @@ namespace
 		}
 	}
 
-	void FProjectLocalizationSettingsDetailsCustomization::UpdateTargetFromReports(ULocalizationTarget* const LocalizationTarget)
+	void FLocalizationTargetSetDetailsCustomization::UpdateTargetFromReports(ULocalizationTarget* const LocalizationTarget)
 	{
 		//TArray< TSharedPtr<IPropertyHandle> > WordCountPropertyHandles;
 
@@ -509,24 +472,24 @@ namespace
 		//}
 	}
 
-	void FProjectLocalizationSettingsDetailsCustomization::CompileAllTargets()
+	void FLocalizationTargetSetDetailsCustomization::CompileAllTargets()
 	{
 		// Execute compile.
 		const TSharedPtr<SWindow> ParentWindow = FSlateApplication::Get().FindWidgetWindow(DetailLayoutBuilder->GetDetailsView().AsShared());
 		TArray<FLocalizationTargetSettings*> TargetSettings;
-		for (ULocalizationTarget* const TargetObject : Settings->TargetObjects)
+		for (ULocalizationTarget* const TargetObject : TargetSet->TargetObjects)
 		{
 			TargetSettings.Add(&TargetObject->Settings);
 		}
 		LocalizationCommandletTasks::CompileTargets(ParentWindow.ToSharedRef(), TargetSettings);
 	}
 
-	TSharedRef<ITableRow> FProjectLocalizationSettingsDetailsCustomization::OnGenerateRow(TSharedPtr<IPropertyHandle> TargetObjectPropertyHandle, const TSharedRef<STableViewBase>& Table)
+	TSharedRef<ITableRow> FLocalizationTargetSetDetailsCustomization::OnGenerateRow(TSharedPtr<IPropertyHandle> TargetObjectPropertyHandle, const TSharedRef<STableViewBase>& Table)
 	{
 		return SNew(SLocalizationDashboardTargetRow, Table, DetailLayoutBuilder->GetPropertyUtilities(), TargetObjectPropertyHandle.ToSharedRef());
 	}
 
-	FReply FProjectLocalizationSettingsDetailsCustomization::OnNewTargetButtonClicked()
+	FReply FLocalizationTargetSetDetailsCustomization::OnNewTargetButtonClicked()
 	{
 		if(TargetObjectsPropertyHandle.IsValid() && TargetObjectsPropertyHandle->IsValidHandle())
 		{
@@ -546,7 +509,7 @@ public:
 	SLATE_BEGIN_ARGS(SLocalizationDashboard) {}
 	SLATE_END_ARGS()
 
-		void Construct(const FArguments& InArgs, const TSharedPtr<SWindow>& OwningWindow, const TSharedRef<SDockTab>& OwningTab, UProjectLocalizationSettings* const Settings);
+		void Construct(const FArguments& InArgs, const TSharedPtr<SWindow>& OwningWindow, const TSharedRef<SDockTab>& OwningTab, ULocalizationTargetSet* const TargetSet);
 	TWeakPtr<SDockTab> ShowTargetEditor(ULocalizationTarget* const LocalizationTarget);
 
 private:
@@ -560,9 +523,9 @@ private:
 const FName SLocalizationDashboard::ProjectTargetSetDetailsTabName("ProjectTargets");
 const FName SLocalizationDashboard::DocumentsTabName("Documents");
 
-void SLocalizationDashboard::Construct(const FArguments& InArgs, const TSharedPtr<SWindow>& OwningWindow, const TSharedRef<SDockTab>& OwningTab, UProjectLocalizationSettings* const Settings)
+void SLocalizationDashboard::Construct(const FArguments& InArgs, const TSharedPtr<SWindow>& OwningWindow, const TSharedRef<SDockTab>& OwningTab, ULocalizationTargetSet* const TargetSet)
 {
-	check(Settings);
+	check(TargetSet);
 
 	TabManager = FGlobalTabmanager::Get()->NewTabManager(OwningTab);
 
@@ -572,7 +535,7 @@ void SLocalizationDashboard::Construct(const FArguments& InArgs, const TSharedPt
 	};
 	TabManager->SetOnPersistLayout(FTabManager::FOnPersistLayout::CreateLambda(PersistLayout));
 
-	const auto& CreateDetailsTab = [Settings](const FSpawnTabArgs& SpawnTabArgs) -> TSharedRef<SDockTab>
+	const auto& CreateDetailsTab = [TargetSet](const FSpawnTabArgs& SpawnTabArgs) -> TSharedRef<SDockTab>
 	{
 		const TSharedRef<SDockTab> DockTab = SNew(SDockTab)
 			.Label(LOCTEXT("ProjectTargetsTabLabel", "Project Targets"));
@@ -581,8 +544,8 @@ void SLocalizationDashboard::Construct(const FArguments& InArgs, const TSharedPt
 		FDetailsViewArgs DetailsViewArgs(false, false, false, FDetailsViewArgs::ENameAreaSettings::HideNameArea, false, nullptr, false, NAME_None);
 		TSharedRef<IDetailsView> DetailsView = PropertyModule.CreateDetailView(DetailsViewArgs);
 
-		DetailsView->RegisterInstancedCustomPropertyLayout(UProjectLocalizationSettings::StaticClass(), FOnGetDetailCustomizationInstance::CreateLambda([&]()->TSharedRef<IDetailCustomization>{return MakeShareable(new FProjectLocalizationSettingsDetailsCustomization(Settings));}));
-		DetailsView->SetObject(Settings, true);
+		DetailsView->RegisterInstancedCustomPropertyLayout(ULocalizationTargetSet::StaticClass(), FOnGetDetailCustomizationInstance::CreateLambda([&]()->TSharedRef<IDetailCustomization>{return MakeShareable(new FLocalizationTargetSetDetailsCustomization(TargetSet));}));
+		DetailsView->SetObject(TargetSet, true);
 
 		DockTab->SetContent(DetailsView);
 
@@ -641,9 +604,9 @@ TWeakPtr<SDockTab> SLocalizationDashboard::ShowTargetEditor(ULocalizationTarget*
 
 	if (!TargetEditorDockTab.IsValid())
 	{
-		UProjectLocalizationSettings* const ProjectSettings = LocalizationTarget->GetTypedOuter<UProjectLocalizationSettings>();
+		ULocalizationTargetSet* const TargetSet = LocalizationTarget->GetTypedOuter<ULocalizationTargetSet>();
 
-		const TSharedRef<SLocalizationTargetEditor> OurTargetEditor = SNew(SLocalizationTargetEditor, ProjectSettings, LocalizationTarget);
+		const TSharedRef<SLocalizationTargetEditor> OurTargetEditor = SNew(SLocalizationTargetEditor, TargetSet, LocalizationTarget);
 		const TSharedRef<SDockTab> NewTargetEditorTab = SNew(SDockTab)
 			.TabRole(ETabRole::DocumentTab)
 			.Label_Lambda( [LocalizationTarget]
@@ -721,7 +684,7 @@ void FLocalizationDashboard::RegisterTabSpawner()
 			.Label(LOCTEXT("MainTabTitle", "Localization Dashboard"))
 			.TabRole(ETabRole::MajorTab);
 
-		DockTab->SetContent( SAssignNew(LocalizationDashboardWidget, SLocalizationDashboard, Args.GetOwnerWindow(), DockTab, GetMutableDefault<UProjectLocalizationSettings>()) );
+		DockTab->SetContent( SAssignNew(LocalizationDashboardWidget, SLocalizationDashboard, Args.GetOwnerWindow(), DockTab, GetMutableDefault<ULocalizationTargetSet>()) );
 
 		return DockTab;
 	};
diff --git a/Engine/Source/Editor/LocalizationDashboard/Private/LocalizationDashboardModule.cpp b/Engine/Source/Editor/LocalizationDashboard/Private/LocalizationDashboardModule.cpp
--- a/Engine/Source/Editor/LocalizationDashboard/Private/LocalizationDashboardModule.cpp
+++ b/Engine/Source/Editor/LocalizationDashboard/Private/LocalizationDashboardModule.cpp
@@ -3,7 +3,7 @@
 #include "LocalizationDashboardPrivatePCH.h"
 #include "ILocalizationDashboardModule.h"
 #include "LocalizationDashboard.h"
-#include "ProjectLocalizationSettings.h"
+#include "LocalizationTarget.h"
 #include "SDockTab.h"
 #include "Features/IModularFeatures.h"
 #include "ILocalizationServiceProvider.h"
@@ -15,7 +15,7 @@ class FLocalizationDashboardModule
 {
 public:
 	FLocalizationDashboardModule()
-		: Settings(GetMutableDefault<UProjectLocalizationSettings>(UProjectLocalizationSettings::StaticClass()))
+		: Settings(GetMutableDefault<ULocalizationTargetSet>(ULocalizationTargetSet::StaticClass()))
 	{
 		
 		for (ULocalizationTarget*& Target : Settings->TargetObjects)
@@ -56,22 +56,23 @@ public:
 
 	virtual ILocalizationServiceProvider* GetCurrentLocalizationServiceProvider() const
 	{
-		if(Settings->LocalizationServiceProvider == NAME_None)
+		if(CurrentServiceProviderName == NAME_None)
 		{
 			return nullptr;
 		}
 
-		auto ServiceProviderNameComparator = [&](ILocalizationServiceProvider* const LSP) -> bool
+		auto ServiceProviderNameComparator = [this](ILocalizationServiceProvider* const LSP) -> bool
 		{
-			return LSP ? LSP->GetName() == Settings->LocalizationServiceProvider : false;
+			return LSP ? LSP->GetName() == CurrentServiceProviderName : false;
 		};
 		ILocalizationServiceProvider* const * LSP = ServiceProviders.FindByPredicate(ServiceProviderNameComparator);
 		return LSP ? *LSP : nullptr;
 	}
 
 private:
-	TWeakObjectPtr<UProjectLocalizationSettings> Settings;
+	TWeakObjectPtr<ULocalizationTargetSet> Settings;
 	TArray<ILocalizationServiceProvider*> ServiceProviders;
+	FName CurrentServiceProviderName;
 };
 
 IMPLEMENT_MODULE(FLocalizationDashboardModule, "LocalizationDashboard");
diff --git a/Engine/Source/Editor/LocalizationDashboard/Private/LocalizationTarget.cpp b/Engine/Source/Editor/LocalizationDashboard/Private/LocalizationTarget.cpp
--- a/Engine/Source/Editor/LocalizationDashboard/Private/LocalizationTarget.cpp
+++ b/Engine/Source/Editor/LocalizationDashboard/Private/LocalizationTarget.cpp
@@ -2,7 +2,7 @@
 
 #include "LocalizationDashboardPrivatePCH.h"
 #include "Classes/LocalizationTarget.h"
-#include "Classes/ProjectLocalizationSettings.h"
+#include "Classes/LocalizationTarget.h"
 #include "ISourceControlModule.h"
 #include "CsvParser.h"
 #include "LocalizationConfigurationScript.h"
@@ -431,20 +431,63 @@ bool FLocalizationTargetSettings::DeleteFiles(const FString* const Culture) cons
 	return HasCompletelySucceeded;
 }
 
-ULocalizationTarget::ULocalizationTarget(const FObjectInitializer& OI)
-	: Super(OI)
+void ULocalizationTarget::PostEditChangeProperty(FPropertyChangedEvent& PropertyChangedEvent)
 {
+	ULocalizationTargetSet* const LocalizationTargetSet = Cast<ULocalizationTargetSet>(GetOuter());
+	if (LocalizationTargetSet)
+	{
+		UClass* ParentClass = LocalizationTargetSet ? LocalizationTargetSet->GetClass() : nullptr;
+		UProperty* ChangedParentProperty = ParentClass ? ParentClass->FindPropertyByName(GET_MEMBER_NAME_CHECKED(ULocalizationTargetSet,TargetObjects)) : nullptr;
+
+		FPropertyChangedEvent PropertyChangedEventForParent(ChangedParentProperty);
+		LocalizationTargetSet->PostEditChangeProperty(PropertyChangedEventForParent);
+	}
 }
 
-void ULocalizationTarget::PostEditChangeProperty(FPropertyChangedEvent& PropertyChangedEvent)
+void ULocalizationTargetSet::PostInitProperties()
 {
-	UProjectLocalizationSettings* const ParentProjectLocalizationSettings = Cast<UProjectLocalizationSettings>(GetOuter());
-	if (ParentProjectLocalizationSettings)
+	Super::PostInitProperties();
+
+	TargetObjects.Empty(Targets.Num());
+	for (const FLocalizationTargetSettings& Target : Targets)
 	{
-		UClass* ParentClass = ParentProjectLocalizationSettings ? ParentProjectLocalizationSettings->GetClass() : nullptr;
-		UProperty* ChangedParentProperty = ParentClass ? ParentClass->FindPropertyByName(GET_MEMBER_NAME_CHECKED(UProjectLocalizationSettings,TargetObjects)) : nullptr;
+		ULocalizationTarget* const NewTargetObject = NewObject<ULocalizationTarget>(this);
+		NewTargetObject->Settings = Target;
+		TargetObjects.Add(NewTargetObject);
+	}
+}
 
-		FPropertyChangedEvent PropertyChangedEventForParent(ChangedParentProperty);
-		ParentProjectLocalizationSettings->PostEditChangeProperty(PropertyChangedEventForParent);
+void ULocalizationTargetSet::PostEditChangeProperty(FPropertyChangedEvent& PropertyChangedEvent)
+{
+	Super::PostEditChangeProperty(PropertyChangedEvent);
+
+	if (PropertyChangedEvent.MemberProperty->GetFName() == GET_MEMBER_NAME_CHECKED(ULocalizationTargetSet,TargetObjects))
+	{
+		Targets.Empty(TargetObjects.Num());
+		for (ULocalizationTarget* const& TargetObject : TargetObjects)
+		{
+			Targets.Add(TargetObject ? TargetObject->Settings : FLocalizationTargetSettings());
+		}
 	}
+
+	ISourceControlModule& SourceControl = ISourceControlModule::Get();
+	ISourceControlProvider& SourceControlProvider = SourceControl.GetProvider();
+	const bool CanUseSourceControl = SourceControl.IsEnabled() && SourceControlProvider.IsEnabled() && SourceControlProvider.IsAvailable();
+
+	if (CanUseSourceControl)
+	{
+		const FString Path = GetConfigFilename(this);
+		const FSourceControlStatePtr SourceControlState = SourceControlProvider.GetState(*Path, EStateCacheUsage::Use);
+
+		// Should only use source control if the file is actually under source control.
+		if (CanUseSourceControl && SourceControlState.IsValid() && !SourceControlState->CanAdd())
+		{
+			if (!SourceControlState->CanEdit())
+			{
+				SourceControlProvider.Execute(ISourceControlOperation::Create<FCheckOut>(), *Path);
+			}
+		}
+	}
+
+	SaveConfig();
 }
\ No newline at end of file
diff --git a/Engine/Source/Editor/LocalizationDashboard/Private/ProjectLocalizationSettings.cpp b/Engine/Source/Editor/LocalizationDashboard/Private/ProjectLocalizationSettings.cpp
deleted file mode 100644
diff --git a/Engine/Source/Editor/LocalizationDashboard/Private/SLocalizationDashboardTargetRow.h b/Engine/Source/Editor/LocalizationDashboard/Private/SLocalizationDashboardTargetRow.h
--- a/Engine/Source/Editor/LocalizationDashboard/Private/SLocalizationDashboardTargetRow.h
+++ b/Engine/Source/Editor/LocalizationDashboard/Private/SLocalizationDashboardTargetRow.h
@@ -4,7 +4,7 @@
 #include "STableRow.h"
 #include "IPropertyUtilities.h"
 #include "PropertyHandle.h"
-#include "ProjectLocalizationSettings.h"
+#include "LocalizationTarget.h"
 #include "LocalizationTarget.h"
 #include "SDockTab.h"
 
diff --git a/Engine/Source/Editor/LocalizationDashboard/Private/SLocalizationTargetEditor.cpp b/Engine/Source/Editor/LocalizationDashboard/Private/SLocalizationTargetEditor.cpp
--- a/Engine/Source/Editor/LocalizationDashboard/Private/SLocalizationTargetEditor.cpp
+++ b/Engine/Source/Editor/LocalizationDashboard/Private/SLocalizationTargetEditor.cpp
@@ -26,7 +26,7 @@ namespace
 	class FLocalizationTargetDetailCustomization : public IDetailCustomization
 	{
 	public:
-		FLocalizationTargetDetailCustomization(UProjectLocalizationSettings* const InProjectSettings, ULocalizationTarget* const InLocalizationTarget);
+		FLocalizationTargetDetailCustomization(ULocalizationTargetSet* const InProjectSettings, ULocalizationTarget* const InLocalizationTarget);
 		virtual void CustomizeDetails(IDetailLayoutBuilder& DetailBuilder) override;
 
 		FLocalizationTargetSettings* GetTargetSettings() const;
@@ -63,7 +63,7 @@ namespace
 		void OnNewSupportedCultureSelected(FCulturePtr SelectedCulture, ESelectInfo::Type SelectInfo);
 
 	private:
-		UProjectLocalizationSettings* ProjectSettings;
+		ULocalizationTargetSet* ProjectSettings;
 		ULocalizationTarget* LocalizationTarget;
 
 		IDetailLayoutBuilder* DetailLayoutBuilder;
@@ -93,7 +93,7 @@ namespace
 		FCulturePtr SelectedNewCulture;
 	};
 
-	FLocalizationTargetDetailCustomization::FLocalizationTargetDetailCustomization(UProjectLocalizationSettings* const InProjectSettings, ULocalizationTarget* const InLocalizationTarget)
+	FLocalizationTargetDetailCustomization::FLocalizationTargetDetailCustomization(ULocalizationTargetSet* const InProjectSettings, ULocalizationTarget* const InLocalizationTarget)
 		: ProjectSettings(InProjectSettings)
 		, LocalizationTarget(InLocalizationTarget)
 		, DetailLayoutBuilder(nullptr)
@@ -924,7 +924,7 @@ namespace
 	}
 }
 
-void SLocalizationTargetEditor::Construct(const FArguments& InArgs, UProjectLocalizationSettings* const InProjectSettings, ULocalizationTarget* const InLocalizationTarget)
+void SLocalizationTargetEditor::Construct(const FArguments& InArgs, ULocalizationTargetSet* const InProjectSettings, ULocalizationTarget* const InLocalizationTarget)
 {
 	check(InProjectSettings->TargetObjects.Contains(InLocalizationTarget));
 
diff --git a/Engine/Source/Editor/LocalizationDashboard/Private/SLocalizationTargetEditor.h b/Engine/Source/Editor/LocalizationDashboard/Private/SLocalizationTargetEditor.h
--- a/Engine/Source/Editor/LocalizationDashboard/Private/SLocalizationTargetEditor.h
+++ b/Engine/Source/Editor/LocalizationDashboard/Private/SLocalizationTargetEditor.h
@@ -4,7 +4,7 @@
 
 #include "SCompoundWidget.h"
 
-class UProjectLocalizationSettings;
+class ULocalizationTargetSet;
 class ULocalizationTarget;
 
 class SLocalizationTargetEditor : public SCompoundWidget
@@ -13,6 +13,6 @@ public:
 	SLATE_BEGIN_ARGS(SLocalizationTargetEditor) {}
 	SLATE_END_ARGS()
 
-	void Construct(const FArguments& InArgs, UProjectLocalizationSettings* const InProjectSettings, ULocalizationTarget* const InTarget);
+	void Construct(const FArguments& InArgs, ULocalizationTargetSet* const InProjectSettings, ULocalizationTarget* const InTarget);
 };
 
diff --git a/Engine/Source/Editor/LocalizationDashboard/Private/SLocalizationTargetStatusButton.h b/Engine/Source/Editor/LocalizationDashboard/Private/SLocalizationTargetStatusButton.h
--- a/Engine/Source/Editor/LocalizationDashboard/Private/SLocalizationTargetStatusButton.h
+++ b/Engine/Source/Editor/LocalizationDashboard/Private/SLocalizationTargetStatusButton.h
@@ -2,7 +2,7 @@
 #pragma once
 
 #include "SCompoundWidget.h"
-#include "ProjectLocalizationSettings.h"
+#include "LocalizationTarget.h"
 #include "LocalizationConfigurationScript.h"
 
 class SLocalizationTargetStatusButton : public SButton
diff --git a/Engine/Source/Editor/LocalizationDashboard/Public/ILocalizationServiceProvider.h b/Engine/Source/Editor/LocalizationDashboard/Public/ILocalizationServiceProvider.h
--- a/Engine/Source/Editor/LocalizationDashboard/Public/ILocalizationServiceProvider.h
+++ b/Engine/Source/Editor/LocalizationDashboard/Public/ILocalizationServiceProvider.h
@@ -6,7 +6,7 @@
 #include "EditorStyleSet.h"
 #include "PropertyHandle.h"
 #include "DetailCategoryBuilder.h"
-#include "ProjectLocalizationSettings.h"
+#include "LocalizationTarget.h"
 
 /**
  * Interface to talking with localization service providers.
