Moved detail customizations from the localization dashboard and localization target editor into their own files. Both are registered and unregistered as part of the module now.

diff --git a/Engine/Source/Editor/LocalizationDashboard/Private/LocalizationDashboard.cpp b/Engine/Source/Editor/LocalizationDashboard/Private/LocalizationDashboard.cpp
--- a/Engine/Source/Editor/LocalizationDashboard/Private/LocalizationDashboard.cpp
+++ b/Engine/Source/Editor/LocalizationDashboard/Private/LocalizationDashboard.cpp
@@ -2,514 +2,25 @@
 
 #include "LocalizationDashboardPrivatePCH.h"
 #include "LocalizationDashboard.h"
-#include "IDetailCustomization.h"
-#include "EditorStyleSet.h"
-#include "DetailLayoutBuilder.h"
-#include "LocalizationTarget.h"
+#include "SDockTab.h"
 #include "PropertyEditorModule.h"
 #include "IDetailsView.h"
-#include "DetailWidgetRow.h"
-#include "DetailCategoryBuilder.h"
-#include "PropertyHandle.h"
-#include "SLocalizationDashboardTargetRow.h"
-#include "TargetsTableEntry.h"
-#include "Runtime/Core/Public/Features/IModularFeatures.h"
-#include "ILocalizationDashboardModule.h"
-#include "ILocalizationServiceProvider.h"
-#include "IPropertyUtilities.h"
-#include "MultiBoxBuilder.h"
-#include "FeedbackContext.h"
-#include "LocalizationCommandletTasks.h"
-#include "FileHelpers.h"
-#include "SLocalizationTargetEditor.h"
 #include "IMainFrameModule.h"
+#include "ModuleManager.h"
+#include "SLocalizationTargetEditor.h"
 
 #define LOCTEXT_NAMESPACE "LocalizationDashboard"
 
 const FName FLocalizationDashboard::TabName("LocalizationDashboard");
 
-class ILocalizationServiceProvider;
-
-namespace
-{
-	struct FLocalizationServiceProviderWrapper;
-
-	class FLocalizationTargetSetDetailsCustomization : public IDetailCustomization
-	{
-	public:
-		FLocalizationTargetSetDetailsCustomization(ULocalizationTargetSet* const InTargetSet);
-		void CustomizeDetails(IDetailLayoutBuilder& DetailBuilder) override;
-
-	private:
-		void BuildTargetsList();
-		void RebuildTargetsList();
-
-		FText GetCurrentServiceProviderDisplayName() const;
-		TSharedRef<SWidget> ServiceProviderComboBox_OnGenerateWidget(TSharedPtr<FLocalizationServiceProviderWrapper> LSPWrapper) const;
-		void ServiceProviderComboBox_OnSelectionChanged(TSharedPtr<FLocalizationServiceProviderWrapper> LSPWrapper, ESelectInfo::Type SelectInfo);
-
-		void GatherAllTargets();
-		void ImportAllTargets();
-		void ExportAllTargets();
-		void UpdateTargetFromReports(ULocalizationTarget* const LocalizationTarget);
-		void CompileAllTargets();
-
-		TSharedRef<ITableRow> OnGenerateRow(TSharedPtr<IPropertyHandle> TargetObjectPropertyHandle, const TSharedRef<STableViewBase>& Table);
-		FReply OnNewTargetButtonClicked();
-
-	private:
-		ULocalizationTargetSet* TargetSet;
-
-		IDetailLayoutBuilder* DetailLayoutBuilder;
-
-		TSharedPtr<IPropertyHandle> ServiceProviderPropertyHandle;
-		IDetailCategoryBuilder* ServiceProviderCategoryBuilder;
-		TArray< TSharedPtr<FLocalizationServiceProviderWrapper> > Providers;
-
-		TSharedPtr<IPropertyHandle> TargetObjectsPropertyHandle;
-		FSimpleDelegate TargetsArrayPropertyHandle_OnNumElementsChanged;
-		TArray< TSharedPtr<IPropertyHandle> > TargetsList;
-		TSharedPtr< SListView< TSharedPtr<IPropertyHandle> > > TargetsListView;
-
-		/* If set, the entry at the index specified needs to be initialized as soon as possible. */
-		int32 NewEntryIndexToBeInitialized;
-	};
-	struct FLocalizationServiceProviderWrapper
-	{
-		FLocalizationServiceProviderWrapper() : Provider(nullptr) {}
-		FLocalizationServiceProviderWrapper(ILocalizationServiceProvider* const InProvider) : Provider(InProvider) {}
-
-		ILocalizationServiceProvider* Provider;
-	};
-
-	FLocalizationTargetSetDetailsCustomization::FLocalizationTargetSetDetailsCustomization(ULocalizationTargetSet* const InTargetSet)
-		: TargetSet(InTargetSet)
-		, DetailLayoutBuilder(nullptr)
-		, ServiceProviderCategoryBuilder(nullptr)
-		, NewEntryIndexToBeInitialized(INDEX_NONE)
-	{
-		TArray<ILocalizationServiceProvider*> ActualProviders = ILocalizationDashboardModule::Get().GetLocalizationServiceProviders();
-		for (ILocalizationServiceProvider* ActualProvider : ActualProviders)
-		{
-			TSharedPtr<FLocalizationServiceProviderWrapper> Provider = MakeShareable(new FLocalizationServiceProviderWrapper(ActualProvider));
-			Providers.Add(Provider);
-		}
-	}
-
-	class FindProviderPredicate
-	{
-	public:
-		FindProviderPredicate(ILocalizationServiceProvider* const InActualProvider)
-			: ActualProvider(InActualProvider)
-		{
-		}
-
-		bool operator()(const TSharedPtr<FLocalizationServiceProviderWrapper>& Provider)
-		{
-			return Provider->Provider == ActualProvider;
-		}
-
-	private:
-		ILocalizationServiceProvider* ActualProvider;
-	};
-
-
-	class FLocalizationDashboardCommands : public TCommands<FLocalizationDashboardCommands>
-	{
-	public:
-		FLocalizationDashboardCommands() 
-			: TCommands<FLocalizationDashboardCommands>("LocalizationDashboard", NSLOCTEXT("Contexts", "LocalizationDashboard", "Localization Dashboard"), NAME_None, FEditorStyle::GetStyleSetName())
-		{
-		}
-
-		TSharedPtr<FUICommandInfo> GatherAllTargets;
-		TSharedPtr<FUICommandInfo> ImportAllTargets;
-		TSharedPtr<FUICommandInfo> ExportAllTargets;
-		TSharedPtr<FUICommandInfo> CompileAllTargets;
-
-		/** Initialize commands */
-		virtual void RegisterCommands() override;
-	};
-
-	void FLocalizationDashboardCommands::RegisterCommands() 
-	{
-		UI_COMMAND( GatherAllTargets, "Gather All", "Gathers text for all targets in the project.", EUserInterfaceActionType::Button, FInputGesture() );
-		UI_COMMAND( ImportAllTargets, "Import All", "Imports translations for all cultures of all targets in the project.", EUserInterfaceActionType::Button, FInputGesture() );
-		UI_COMMAND( ExportAllTargets, "Export All", "Exports translations for all cultures of all targets in the project.", EUserInterfaceActionType::Button, FInputGesture() );
-		UI_COMMAND( CompileAllTargets, "Compile All", "Compiles translations for all targets in the project.", EUserInterfaceActionType::Button, FInputGesture() );
-	}
-
-	void FLocalizationTargetSetDetailsCustomization::CustomizeDetails(IDetailLayoutBuilder& DetailBuilder)
-	{
-		DetailLayoutBuilder = &DetailBuilder;
-
-		{
-			IDetailCategoryBuilder& DetailCategoryBuilder = DetailBuilder.EditCategory(FName("Targets"));
-			TargetObjectsPropertyHandle = DetailBuilder.GetProperty(GET_MEMBER_NAME_CHECKED(ULocalizationTargetSet,TargetObjects));
-			if (TargetObjectsPropertyHandle->IsValidHandle())
-			{
-				TargetObjectsPropertyHandle->MarkHiddenByCustomization();
-				TargetsArrayPropertyHandle_OnNumElementsChanged = FSimpleDelegate::CreateSP(this, &FLocalizationTargetSetDetailsCustomization::RebuildTargetsList);
-				TargetObjectsPropertyHandle->AsArray()->SetOnNumElementsChanged(TargetsArrayPropertyHandle_OnNumElementsChanged);
-
-				FLocalizationDashboardCommands::Register();
-				const TSharedRef< FUICommandList > CommandList = MakeShareable(new FUICommandList);
-				FToolBarBuilder ToolBarBuilder( CommandList, FMultiBoxCustomization::AllowCustomization("LocalizationDashboard") );
-
-				CommandList->MapAction( FLocalizationDashboardCommands::Get().GatherAllTargets, FExecuteAction::CreateSP(this, &FLocalizationTargetSetDetailsCustomization::GatherAllTargets));
-				ToolBarBuilder.AddToolBarButton(FLocalizationDashboardCommands::Get().GatherAllTargets, NAME_None, TAttribute<FText>(), TAttribute<FText>(), FSlateIcon(FEditorStyle::GetStyleSetName(), "LocalizationDashboard.GatherAllTargets"));
-
-				CommandList->MapAction( FLocalizationDashboardCommands::Get().ImportAllTargets, FExecuteAction::CreateSP(this, &FLocalizationTargetSetDetailsCustomization::ImportAllTargets));
-				ToolBarBuilder.AddToolBarButton(FLocalizationDashboardCommands::Get().ImportAllTargets, NAME_None, TAttribute<FText>(), TAttribute<FText>(), FSlateIcon(FEditorStyle::GetStyleSetName(), "LocalizationDashboard.ImportForAllTargetsCultures"));
-
-				CommandList->MapAction( FLocalizationDashboardCommands::Get().ExportAllTargets, FExecuteAction::CreateSP(this, &FLocalizationTargetSetDetailsCustomization::ExportAllTargets));
-				ToolBarBuilder.AddToolBarButton(FLocalizationDashboardCommands::Get().ExportAllTargets, NAME_None, TAttribute<FText>(), TAttribute<FText>(), FSlateIcon(FEditorStyle::GetStyleSetName(), "LocalizationDashboard.ExportForAllTargetsCultures"));
-
-				CommandList->MapAction( FLocalizationDashboardCommands::Get().CompileAllTargets, FExecuteAction::CreateSP(this, &FLocalizationTargetSetDetailsCustomization::CompileAllTargets));
-				ToolBarBuilder.AddToolBarButton(FLocalizationDashboardCommands::Get().CompileAllTargets, NAME_None, TAttribute<FText>(), TAttribute<FText>(), FSlateIcon(FEditorStyle::GetStyleSetName(), "LocalizationDashboard.CompileAllTargets"));
-
-				BuildTargetsList();
-
-				DetailCategoryBuilder.AddCustomRow( LOCTEXT("TargetsFilterString", "Targets") )
-					.WholeRowContent()
-					[
-						SNew(SVerticalBox)
-						+SVerticalBox::Slot()
-						.AutoHeight()
-						[
-							ToolBarBuilder.MakeWidget()
-						]
-						+SVerticalBox::Slot()
-							.AutoHeight()
-							[
-								SAssignNew(TargetsListView, SListView< TSharedPtr<IPropertyHandle> >)
-								.OnGenerateRow(this, &FLocalizationTargetSetDetailsCustomization::OnGenerateRow)
-								.ListItemsSource(&TargetsList)
-								.SelectionMode(ESelectionMode::None)
-								.HeaderRow
-								(
-								SNew(SHeaderRow)
-								+SHeaderRow::Column("Target")
-								.DefaultLabel( LOCTEXT("TargetColumnLabel", "Target"))
-								.HAlignHeader(HAlign_Left)
-								.HAlignCell(HAlign_Left)
-								.VAlignCell(VAlign_Center)
-								+SHeaderRow::Column("Status")
-								.DefaultLabel( LOCTEXT("StatusColumnLabel", "Status"))
-								.HAlignHeader(HAlign_Center)
-								.HAlignCell(HAlign_Center)
-								.VAlignCell(VAlign_Center)
-								+SHeaderRow::Column("Cultures")
-								.DefaultLabel( LOCTEXT("CulturesColumnLabel", "Cultures"))
-								.HAlignHeader(HAlign_Fill)
-								.HAlignCell(HAlign_Fill)
-								.VAlignCell(VAlign_Center)
-								+SHeaderRow::Column("WordCount")
-								.DefaultLabel( LOCTEXT("WordCountColumnLabel", "Word Count"))
-								.HAlignHeader(HAlign_Center)
-								.HAlignCell(HAlign_Center)
-								.VAlignCell(VAlign_Center)
-								+SHeaderRow::Column("Actions")
-								.DefaultLabel( LOCTEXT("ActionsColumnLabel", "Actions"))
-								.HAlignHeader(HAlign_Center)
-								.HAlignCell(HAlign_Center)
-								.VAlignCell(VAlign_Center)
-								)
-							]
-						+SVerticalBox::Slot()
-							.AutoHeight()
-							.VAlign(VAlign_Center)
-							[
-								SNew(SButton)
-								.Text(LOCTEXT("AddNewTargetButtonLabel", "Add New Target"))
-								.OnClicked(this, &FLocalizationTargetSetDetailsCustomization::OnNewTargetButtonClicked)
-							]
-					];
-			}
-		}
-	}
-
-	void FLocalizationTargetSetDetailsCustomization::BuildTargetsList()
-	{
-		const TSharedPtr<IPropertyHandleArray> TargetObjectsArrayPropertyHandle = TargetObjectsPropertyHandle->AsArray();
-		if (TargetObjectsArrayPropertyHandle.IsValid())
-		{
-			uint32 ElementCount = 0;
-			TargetObjectsArrayPropertyHandle->GetNumElements(ElementCount);
-			for (uint32 i = 0; i < ElementCount; ++i)
-			{
-				const TSharedPtr<IPropertyHandle> TargetObjectPropertyHandle = TargetObjectsArrayPropertyHandle->GetElement(i);
-				TargetsList.Add(TargetObjectPropertyHandle);
-			}
-		}
-	}
-
-	void FLocalizationTargetSetDetailsCustomization::RebuildTargetsList()
-	{
-		const TSharedPtr<IPropertyHandleArray> TargetObjectsArrayPropertyHandle = TargetObjectsPropertyHandle->AsArray();
-		if (TargetObjectsArrayPropertyHandle.IsValid() && NewEntryIndexToBeInitialized != INDEX_NONE)
-		{
-			const TSharedPtr<IPropertyHandle> TargetObjectPropertyHandle = TargetObjectsArrayPropertyHandle->GetElement(NewEntryIndexToBeInitialized);
-			if (TargetObjectPropertyHandle.IsValid() && TargetObjectPropertyHandle->IsValidHandle())
-			{
-				ULocalizationTarget* const NewTarget = NewObject<ULocalizationTarget>(TargetSet);
-
-				TArray<void*> RawData;
-				TargetObjectsPropertyHandle->AccessRawData(RawData);
-				void* RawDatum = RawData.Top();
-				TArray<ULocalizationTarget*>& TargetObjectsArray = *(reinterpret_cast< TArray<ULocalizationTarget*>* >(RawDatum));
-
-				FString NewTargetName = "NewTarget";
-				auto TargetNameIsIdentical = [&](ULocalizationTarget* Target) -> bool
-				{
-					return (Target != NewTarget) && Target && (Target->Settings.Name == NewTargetName);
-				};
-
-				for (uint32 i = 0; TargetObjectsArray.ContainsByPredicate(TargetNameIsIdentical); ++i)
-				{
-					NewTargetName = FString::Printf(TEXT("NewTarget%u"), i);
-				}
-
-				NewTarget->Settings.Name = NewTargetName;
-				NewTarget->Settings.NativeCultureStatistics.CultureName = FInternationalization::Get().GetCurrentCulture()->GetName();
-
-				const UObject* SetValue = NewTarget;
-				TargetObjectPropertyHandle->SetValue(SetValue);
-
-				NewEntryIndexToBeInitialized = INDEX_NONE;
-			}
-		}
-
-		TargetsList.Empty();
-		BuildTargetsList();
-
-		if (TargetsListView.IsValid())
-		{
-			TargetsListView->RequestListRefresh();
-		}
-	}
-
-	FText FLocalizationTargetSetDetailsCustomization::GetCurrentServiceProviderDisplayName() const
-	{
-		ILocalizationServiceProvider* const LSP = ILocalizationDashboardModule::Get().GetCurrentLocalizationServiceProvider();
-		return LSP ? LSP->GetDisplayName() : LOCTEXT("NoServiceProviderName", "None");
-	}
-
-	TSharedRef<SWidget> FLocalizationTargetSetDetailsCustomization::ServiceProviderComboBox_OnGenerateWidget(TSharedPtr<FLocalizationServiceProviderWrapper> LSPWrapper) const
-	{
-		ILocalizationServiceProvider* const LSP = LSPWrapper->Provider;
-		return	SNew(STextBlock)
-			.Text(LSP ? LSP->GetDisplayName() : LOCTEXT("NoServiceProviderName", "None"));
-	}
-
-	void FLocalizationTargetSetDetailsCustomization::ServiceProviderComboBox_OnSelectionChanged(TSharedPtr<FLocalizationServiceProviderWrapper> LSPWrapper, ESelectInfo::Type SelectInfo)
-	{
-		ILocalizationServiceProvider* const LSP = LSPWrapper->Provider;
-		FString ServiceProviderName = LSP ? LSP->GetName().ToString() : TEXT("None");
-		ServiceProviderPropertyHandle->SetValue(ServiceProviderName);
-
-		if (LSP)
-		{
-			LSP->CustomizeSettingsDetails(*ServiceProviderCategoryBuilder);
-		}
-		DetailLayoutBuilder->ForceRefreshDetails();
-	}
-
-	void FLocalizationTargetSetDetailsCustomization::GatherAllTargets()
-	{
-		for (ULocalizationTarget* const LocalizationTarget : TargetSet->TargetObjects)
-		{
-			// Save unsaved packages.
-			const bool bPromptUserToSave = true;
-			const bool bSaveMapPackages = true;
-			const bool bSaveContentPackages = true;
-			const bool bFastSave = false;
-			const bool bNotifyNoPackagesSaved = false;
-			const bool bCanBeDeclined = true;
-			bool DidPackagesNeedSaving;
-			const bool WerePackagesSaved = FEditorFileUtils::SaveDirtyPackages(bPromptUserToSave, bSaveMapPackages, bSaveContentPackages, bFastSave, bNotifyNoPackagesSaved, bCanBeDeclined, &DidPackagesNeedSaving);
-
-			if (DidPackagesNeedSaving && !WerePackagesSaved)
-			{
-				// Give warning dialog.
-				const FText MessageText = NSLOCTEXT("LocalizationCultureActions", "UnsavedPackagesWarningDialogMessage", "There are unsaved changes. These changes may not be gathered from correctly.");
-				const FText TitleText = NSLOCTEXT("LocalizationCultureActions", "UnsavedPackagesWarningDialogTitle", "Unsaved Changes Before Gather");
-				switch(FMessageDialog::Open(EAppMsgType::OkCancel, MessageText, &TitleText))
-				{
-				case EAppReturnType::Cancel:
-					{
-						return;
-					}
-					break;
-				}
-			}
-
-		}
-
-		// Execute gather.
-		const TSharedPtr<SWindow> ParentWindow = FSlateApplication::Get().FindWidgetWindow(DetailLayoutBuilder->GetDetailsView().AsShared());
-		TArray<FLocalizationTargetSettings*> TargetSettings;
-		for (ULocalizationTarget* const TargetObject : TargetSet->TargetObjects)
-		{
-			TargetSettings.Add(&TargetObject->Settings);
-		}
-		LocalizationCommandletTasks::GatherTargets(ParentWindow.ToSharedRef(), TargetSettings);
-
-		for (ULocalizationTarget* const LocalizationTarget : TargetSet->TargetObjects)
-		{
-			UpdateTargetFromReports(LocalizationTarget);
-		}
-	}
-
-	void FLocalizationTargetSetDetailsCustomization::ImportAllTargets()
-	{
-		IDesktopPlatform* DesktopPlatform = FDesktopPlatformModule::Get();
-		if (DesktopPlatform)
-		{
-			void* ParentWindowWindowHandle = NULL;
-			const TSharedPtr<SWindow> ParentWindow = FSlateApplication::Get().FindWidgetWindow(DetailLayoutBuilder->GetDetailsView().AsShared());
-			if (ParentWindow.IsValid() && ParentWindow->GetNativeWindow().IsValid())
-			{
-				ParentWindowWindowHandle = ParentWindow->GetNativeWindow()->GetOSWindowHandle();
-			}
-
-			
-			const FString DefaultPath = FPaths::ConvertRelativePathToFull(FPaths::GameContentDir() / TEXT("Localization"));
-
-			// Prompt the user for the directory
-			FString OutputDirectory;
-			if (DesktopPlatform->OpenDirectoryDialog(ParentWindowWindowHandle, LOCTEXT("ImportAllTranslationsDialogTitle", "Import All Translations from Directory").ToString(), DefaultPath, OutputDirectory))
-			{
-				TArray<FLocalizationTargetSettings*> TargetSettings;
-				for (ULocalizationTarget* const TargetObject : TargetSet->TargetObjects)
-				{
-					TargetSettings.Add(&TargetObject->Settings);
-				}
-				LocalizationCommandletTasks::ImportTargets(ParentWindow.ToSharedRef(), TargetSettings, TOptional<FString>(OutputDirectory));
-
-				for (ULocalizationTarget* const LocalizationTarget : TargetSet->TargetObjects)
-				{
-					UpdateTargetFromReports(LocalizationTarget);
-				}
-			}
-		}
-	}
-
-	void FLocalizationTargetSetDetailsCustomization::ExportAllTargets()
-	{
-		IDesktopPlatform* DesktopPlatform = FDesktopPlatformModule::Get();
-		if (DesktopPlatform)
-		{
-			void* ParentWindowWindowHandle = NULL;
-			const TSharedPtr<SWindow> ParentWindow = FSlateApplication::Get().FindWidgetWindow(DetailLayoutBuilder->GetDetailsView().AsShared());
-			if (ParentWindow.IsValid() && ParentWindow->GetNativeWindow().IsValid())
-			{
-				ParentWindowWindowHandle = ParentWindow->GetNativeWindow()->GetOSWindowHandle();
-			}
-
-			const FString DefaultPath = FPaths::ConvertRelativePathToFull(FPaths::GameContentDir() / TEXT("Localization"));
-
-			// Prompt the user for the directory
-			FString OutputDirectory;
-			if (DesktopPlatform->OpenDirectoryDialog(ParentWindowWindowHandle, LOCTEXT("ExportAllTranslationsDialogTitle", "Export All Translations to Directory").ToString(), DefaultPath, OutputDirectory))
-			{
-				TArray<FLocalizationTargetSettings*> TargetSettings;
-				for (ULocalizationTarget* const TargetObject : TargetSet->TargetObjects)
-				{
-					TargetSettings.Add(&TargetObject->Settings);
-				}
-				LocalizationCommandletTasks::ExportTargets(ParentWindow.ToSharedRef(), TargetSettings, TOptional<FString>(OutputDirectory));
-			}
-		}
-	}
-
-	void FLocalizationTargetSetDetailsCustomization::UpdateTargetFromReports(ULocalizationTarget* const LocalizationTarget)
-	{
-		//TArray< TSharedPtr<IPropertyHandle> > WordCountPropertyHandles;
-
-		//const TSharedPtr<IPropertyHandle> TargetSettingsPropertyHandle = TargetEditor->GetTargetSettingsPropertyHandle();
-		//if (TargetSettingsPropertyHandle.IsValid() && TargetSettingsPropertyHandle->IsValidHandle())
-		//{
-		//	const TSharedPtr<IPropertyHandle> NativeCultureStatisticsPropertyHandle = TargetSettingsPropertyHandle->GetChildHandle(GET_MEMBER_NAME_CHECKED(FLocalizationTargetSettings, NativeCultureStatistics));
-		//	if (NativeCultureStatisticsPropertyHandle.IsValid() && NativeCultureStatisticsPropertyHandle->IsValidHandle())
-		//	{
-		//		const TSharedPtr<IPropertyHandle> WordCountPropertyHandle = NativeCultureStatisticsPropertyHandle->GetChildHandle(GET_MEMBER_NAME_CHECKED(FCultureStatistics, WordCount));
-		//		if (WordCountPropertyHandle.IsValid() && WordCountPropertyHandle->IsValidHandle())
-		//		{
-		//			WordCountPropertyHandles.Add(WordCountPropertyHandle);
-		//		}
-		//	}
-
-		//	const TSharedPtr<IPropertyHandle> SupportedCulturesStatisticsPropertyHandle = TargetSettingsPropertyHandle->GetChildHandle(GET_MEMBER_NAME_CHECKED(FLocalizationTargetSettings, SupportedCulturesStatistics));
-		//	if (SupportedCulturesStatisticsPropertyHandle.IsValid() && SupportedCulturesStatisticsPropertyHandle->IsValidHandle())
-		//	{
-		//		uint32 SupportedCultureCount = 0;
-		//		SupportedCulturesStatisticsPropertyHandle->GetNumChildren(SupportedCultureCount);
-		//		for (uint32 i = 0; i < SupportedCultureCount; ++i)
-		//		{
-		//			const TSharedPtr<IPropertyHandle> ElementPropertyHandle = SupportedCulturesStatisticsPropertyHandle->GetChildHandle(i);
-		//			if (ElementPropertyHandle.IsValid() && ElementPropertyHandle->IsValidHandle())
-		//			{
-		//				const TSharedPtr<IPropertyHandle> WordCountPropertyHandle = SupportedCulturesStatisticsPropertyHandle->GetChildHandle(GET_MEMBER_NAME_CHECKED(FCultureStatistics, WordCount));
-		//				if (WordCountPropertyHandle.IsValid() && WordCountPropertyHandle->IsValidHandle())
-		//				{
-		//					WordCountPropertyHandles.Add(WordCountPropertyHandle);
-		//				}
-		//			}
-		//		}
-		//	}
-		//}
-
-		//for (const TSharedPtr<IPropertyHandle>& WordCountPropertyHandle : WordCountPropertyHandles)
-		//{
-		//	WordCountPropertyHandle->NotifyPreChange();
-		//}
-		LocalizationTarget->Settings.UpdateWordCountsFromCSV();
-		LocalizationTarget->Settings.UpdateStatusFromConflictReport();
-		//for (const TSharedPtr<IPropertyHandle>& WordCountPropertyHandle : WordCountPropertyHandles)
-		//{
-		//	WordCountPropertyHandle->NotifyPostChange();
-		//}
-	}
-
-	void FLocalizationTargetSetDetailsCustomization::CompileAllTargets()
-	{
-		// Execute compile.
-		const TSharedPtr<SWindow> ParentWindow = FSlateApplication::Get().FindWidgetWindow(DetailLayoutBuilder->GetDetailsView().AsShared());
-		TArray<FLocalizationTargetSettings*> TargetSettings;
-		for (ULocalizationTarget* const TargetObject : TargetSet->TargetObjects)
-		{
-			TargetSettings.Add(&TargetObject->Settings);
-		}
-		LocalizationCommandletTasks::CompileTargets(ParentWindow.ToSharedRef(), TargetSettings);
-	}
-
-	TSharedRef<ITableRow> FLocalizationTargetSetDetailsCustomization::OnGenerateRow(TSharedPtr<IPropertyHandle> TargetObjectPropertyHandle, const TSharedRef<STableViewBase>& Table)
-	{
-		return SNew(SLocalizationDashboardTargetRow, Table, DetailLayoutBuilder->GetPropertyUtilities(), TargetObjectPropertyHandle.ToSharedRef());
-	}
-
-	FReply FLocalizationTargetSetDetailsCustomization::OnNewTargetButtonClicked()
-	{
-		if(TargetObjectsPropertyHandle.IsValid() && TargetObjectsPropertyHandle->IsValidHandle())
-		{
-			uint32 NewElementIndex;
-			TargetObjectsPropertyHandle->AsArray()->GetNumElements(NewElementIndex);
-			TargetObjectsPropertyHandle->AsArray()->AddItem();
-			NewEntryIndexToBeInitialized = NewElementIndex;
-		}
-
-		return FReply::Handled();
-	}
-}
-
 class SLocalizationDashboard : public SCompoundWidget
 {
 public:
 	SLATE_BEGIN_ARGS(SLocalizationDashboard) {}
 	SLATE_END_ARGS()
 
-		void Construct(const FArguments& InArgs, const TSharedPtr<SWindow>& OwningWindow, const TSharedRef<SDockTab>& OwningTab, ULocalizationTargetSet* const TargetSet);
+	void Construct(const FArguments& InArgs, const TSharedPtr<SWindow>& OwningWindow, const TSharedRef<SDockTab>& OwningTab, ULocalizationTargetSet* const TargetSet);
+
 	TWeakPtr<SDockTab> ShowTargetEditor(ULocalizationTarget* const LocalizationTarget);
 
 private:
@@ -544,7 +55,6 @@ void SLocalizationDashboard::Construct(const FArguments& InArgs, const TSharedPt
 		FDetailsViewArgs DetailsViewArgs(false, false, false, FDetailsViewArgs::ENameAreaSettings::HideNameArea, false, nullptr, false, NAME_None);
 		TSharedRef<IDetailsView> DetailsView = PropertyModule.CreateDetailView(DetailsViewArgs);
 
-		DetailsView->RegisterInstancedCustomPropertyLayout(ULocalizationTargetSet::StaticClass(), FOnGetDetailCustomizationInstance::CreateLambda([&]()->TSharedRef<IDetailCustomization>{return MakeShareable(new FLocalizationTargetSetDetailsCustomization(TargetSet));}));
 		DetailsView->SetObject(TargetSet, true);
 
 		DockTab->SetContent(DetailsView);
diff --git a/Engine/Source/Editor/LocalizationDashboard/Private/LocalizationDashboardModule.cpp b/Engine/Source/Editor/LocalizationDashboard/Private/LocalizationDashboardModule.cpp
--- a/Engine/Source/Editor/LocalizationDashboard/Private/LocalizationDashboardModule.cpp
+++ b/Engine/Source/Editor/LocalizationDashboard/Private/LocalizationDashboardModule.cpp
@@ -7,6 +7,8 @@
 #include "SDockTab.h"
 #include "Features/IModularFeatures.h"
 #include "ILocalizationServiceProvider.h"
+#include "LocalizationTargetSetDetailCustomization.h"
+#include "LocalizationTargetDetailCustomization.h"
 
 #define LOCTEXT_NAMESPACE "LocalizationDashboard"
 
@@ -28,15 +30,33 @@ public:
 	// Begin IModuleInterface interface
 	virtual void StartupModule() override
 	{
-		FLocalizationDashboard::Initialize();
-
 		ServiceProviders = IModularFeatures::Get().GetModularFeatureImplementations<ILocalizationServiceProvider>("LocalizationService");
 		ServiceProviders.Insert(nullptr, 0); // "None"
+
+		FPropertyEditorModule& PropertyModule = FModuleManager::LoadModuleChecked<FPropertyEditorModule>("PropertyEditor");
+		PropertyModule.RegisterCustomClassLayout("LocalizationTargetSet", FOnGetDetailCustomizationInstance::CreateLambda(
+			[]() -> TSharedRef<IDetailCustomization>
+			{
+				return MakeShareable( new FLocalizationTargetSetDetailCustomization());
+			})
+		);
+		PropertyModule.RegisterCustomClassLayout("LocalizationTarget", FOnGetDetailCustomizationInstance::CreateLambda(
+			[]() -> TSharedRef<IDetailCustomization>
+			{
+				return MakeShareable( new FLocalizationTargetDetailCustomization());
+			})
+		);
+
+		FLocalizationDashboard::Initialize();
 	}
 
 	virtual void ShutdownModule() override
 	{
 		FLocalizationDashboard::Terminate();
+
+		FPropertyEditorModule& PropertyModule = FModuleManager::LoadModuleChecked<FPropertyEditorModule>("PropertyEditor");
+		PropertyModule.UnregisterCustomClassLayout("LocalizationTarget");
+		PropertyModule.UnregisterCustomClassLayout("LocalizationTargetSet");
 	}
 	// End IModuleInterface interface
 
diff --git a/Engine/Source/Editor/LocalizationDashboard/Private/LocalizationTargetDetailCustomization.cpp b/Engine/Source/Editor/LocalizationDashboard/Private/LocalizationTargetDetailCustomization.cpp
new file mode 100644
--- /dev/null
+++ b/Engine/Source/Editor/LocalizationDashboard/Private/LocalizationTargetDetailCustomization.cpp
@@ -0,0 +1,850 @@
+// Copyright 1998-2015 Epic Games, Inc. All Rights Reserved.
+
+#include "LocalizationDashboardPrivatePCH.h"
+#include "LocalizationTargetDetailCustomization.h"
+#include "DetailLayoutBuilder.h"
+#include "LocalizationTarget.h"
+#include "DetailWidgetRow.h"
+#include "SLocalizationTargetStatusButton.h"
+#include "SLocalizationTargetEditorCultureRow.h"
+#include "SCulturePicker.h"
+#include "IDetailPropertyRow.h"
+
+#define LOCTEXT_NAMESPACE "LocalizationTargetEditor"
+
+FLocalizationTargetDetailCustomization::FLocalizationTargetDetailCustomization()
+	: DetailLayoutBuilder(nullptr)
+	, NewEntryIndexToBeInitialized(INDEX_NONE)
+{
+}
+
+class FLocalizationTargetEditorCommands : public TCommands<FLocalizationTargetEditorCommands>
+{
+public:
+	FLocalizationTargetEditorCommands() 
+		: TCommands<FLocalizationTargetEditorCommands>("LocalizationTargetEditor", NSLOCTEXT("Contexts", "LocalizationTargetEditor", "Localization Target Editor"), NAME_None, FEditorStyle::GetStyleSetName())
+	{
+	}
+
+	TSharedPtr<FUICommandInfo> Gather;
+	TSharedPtr<FUICommandInfo> ImportAllCultures;
+	TSharedPtr<FUICommandInfo> ExportAllCultures;
+	TSharedPtr<FUICommandInfo> RefreshWordCounts;
+	TSharedPtr<FUICommandInfo> Compile;
+
+	/** Initialize commands */
+	virtual void RegisterCommands() override;
+};
+
+void FLocalizationTargetEditorCommands::RegisterCommands() 
+{
+	UI_COMMAND( Gather, "Gather", "Gathers text for this target.", EUserInterfaceActionType::Button, FInputGesture() );
+	UI_COMMAND( ImportAllCultures, "Import All", "Imports translations for all cultures of this target.", EUserInterfaceActionType::Button, FInputGesture() );
+	UI_COMMAND( ExportAllCultures, "Export All", "Exports translations for all cultures of this target.", EUserInterfaceActionType::Button, FInputGesture() );
+	UI_COMMAND( RefreshWordCounts, "Count Words", "Recalculates the word counts for all cultures of this target.", EUserInterfaceActionType::Button, FInputGesture() );
+	UI_COMMAND( Compile, "Compile", "Compiles translations for this target.", EUserInterfaceActionType::Button, FInputGesture() );
+}
+
+void FLocalizationTargetDetailCustomization::CustomizeDetails(IDetailLayoutBuilder& DetailBuilder)
+{
+	DetailLayoutBuilder = &DetailBuilder;
+	{
+		TArray< TWeakObjectPtr<UObject> > ObjectsBeingCustomized;
+		DetailLayoutBuilder->GetObjectsBeingCustomized(ObjectsBeingCustomized);
+		LocalizationTarget = CastChecked<ULocalizationTarget>(ObjectsBeingCustomized.Top().Get());
+		TargetSet = CastChecked<ULocalizationTargetSet>(LocalizationTarget->GetOuter());
+	}
+
+	TargetSettingsPropertyHandle = DetailBuilder.GetProperty(GET_MEMBER_NAME_CHECKED(ULocalizationTarget, Settings));
+
+	if (TargetSettingsPropertyHandle.IsValid() && TargetSettingsPropertyHandle->IsValidHandle())
+	{
+		TargetSettingsPropertyHandle->MarkHiddenByCustomization();
+	}
+
+	//{
+	//	IDetailCategoryBuilder& ServiceProviderCategoryBuilder = DetailBuilder.EditCategory(FName("LocalizationServiceProvider"));
+
+	//	ILocalizationServiceProvider* const LSP = ILocalizationDashboardModule::Get().GetCurrentLocalizationServiceProvider();
+	//	if (LSP && LocalizationTarget)
+	//	{
+	//		LSP->CustomizeTargetDetails(ServiceProviderCategoryBuilder, *LocalizationTarget);
+	//	}
+	//}
+
+	if (TargetSettingsPropertyHandle.IsValid() && TargetSettingsPropertyHandle->IsValidHandle())
+	{
+		IDetailCategoryBuilder& DetailCategoryBuilder = DetailBuilder.EditCategory(FName("Target"));
+
+		const TSharedPtr<IPropertyHandle> NamePropertyHandle = TargetSettingsPropertyHandle->GetChildHandle(GET_MEMBER_NAME_CHECKED(FLocalizationTargetSettings, Name));
+		if (NamePropertyHandle.IsValid() && NamePropertyHandle->IsValidHandle())
+		{
+			FDetailWidgetRow& DetailWidgetRow = DetailCategoryBuilder.AddCustomRow( LOCTEXT("TargetNameFilterString", "Target Name") );
+			DetailWidgetRow.NameContent()
+				[
+					NamePropertyHandle->CreatePropertyNameWidget()
+				];
+			DetailWidgetRow.ValueContent()
+				[
+					SAssignNew(TargetNameEditableTextBox, SEditableTextBox)
+					.Text(this, &FLocalizationTargetDetailCustomization::GetTargetName)
+					.RevertTextOnEscape(true)
+					.OnTextChanged(this, &FLocalizationTargetDetailCustomization::OnTargetNameChanged)
+					.OnTextCommitted(this, &FLocalizationTargetDetailCustomization::OnTargetNameCommitted)
+				];
+		}
+		const TSharedPtr<IPropertyHandle> StatusPropertyHandle = TargetSettingsPropertyHandle->GetChildHandle(GET_MEMBER_NAME_CHECKED(FLocalizationTargetSettings, Status));
+		if (StatusPropertyHandle.IsValid() && StatusPropertyHandle->IsValidHandle())
+		{
+			StatusPropertyHandle->MarkHiddenByCustomization();
+			if (LocalizationTarget.IsValid())
+			{
+				FDetailWidgetRow& StatusRow = DetailCategoryBuilder.AddCustomRow( LOCTEXT("ConflictReportStatusFilterString", "Conflict Report Status") );
+				StatusRow.NameContent()
+					[
+						StatusPropertyHandle->CreatePropertyNameWidget()
+					];
+				StatusRow.ValueContent()
+					.HAlign(HAlign_Left)
+					.VAlign(VAlign_Center)
+					[
+						SNew(SLocalizationTargetStatusButton, *LocalizationTarget)
+					];
+			}
+		}
+
+		const TSharedPtr<IPropertyHandle> TargetDependenciesPropertyHandle = TargetSettingsPropertyHandle->GetChildHandle(GET_MEMBER_NAME_CHECKED(FLocalizationTargetSettings, TargetDependencies));
+		if (TargetDependenciesPropertyHandle.IsValid() && TargetDependenciesPropertyHandle->IsValidHandle())
+		{
+			TargetDependenciesPropertyHandle->MarkHiddenByCustomization();
+			if (LocalizationTarget.IsValid())
+			{
+				FDetailWidgetRow& TargetDependenciesRow = DetailCategoryBuilder.AddCustomRow( LOCTEXT("ConflictReportStatusFilterString", "Conflict Report Status") );
+				TargetDependenciesRow.NameContent()
+					[
+						TargetDependenciesPropertyHandle->CreatePropertyNameWidget()
+					];
+				TargetDependenciesRow.ValueContent()
+					.HAlign(HAlign_Left)
+					.VAlign(VAlign_Center)
+					[
+						SNew(SComboButton)
+						.ButtonContent()
+						[
+							SAssignNew(TargetDependenciesHorizontalBox, SHorizontalBox)
+						]
+						.HasDownArrow(true)
+							.OnGetMenuContent_Lambda([this]() -> TSharedRef<SWidget>
+						{
+							RebuildTargetsList();
+
+							if (TargetDependenciesOptionsList.Num() > 0)
+							{
+								return SNew(SBox)
+									.MaxDesiredHeight(400.0f)
+									.MaxDesiredWidth(300.0f)
+									[
+										SNew(SListView<ULocalizationTarget*>)
+										.SelectionMode(ESelectionMode::None)
+										.ListItemsSource(&TargetDependenciesOptionsList)
+										.OnGenerateRow(this, &FLocalizationTargetDetailCustomization::OnGenerateTargetRow)
+									];
+							}
+							else
+							{
+								return SNullWidget::NullWidget;
+							}
+						})
+					];
+
+				RebuildTargetDependenciesBox();
+			}
+		}
+
+		const TSharedPtr<IPropertyHandle> AdditionalManifestDependenciesPropertyHandle = TargetSettingsPropertyHandle->GetChildHandle(GET_MEMBER_NAME_CHECKED(FLocalizationTargetSettings, AdditionalManifestDependencies));
+		if (AdditionalManifestDependenciesPropertyHandle.IsValid() && AdditionalManifestDependenciesPropertyHandle->IsValidHandle())
+		{
+			DetailCategoryBuilder.AddProperty(AdditionalManifestDependenciesPropertyHandle);
+		}
+	}
+
+	if (TargetSettingsPropertyHandle.IsValid() && TargetSettingsPropertyHandle->IsValidHandle())
+	{
+		IDetailCategoryBuilder& DetailCategoryBuilder = DetailBuilder.EditCategory(FName("Cultures"));
+
+		NativeCultureStatisticsPropertyHandle = TargetSettingsPropertyHandle->GetChildHandle(GET_MEMBER_NAME_CHECKED(FLocalizationTargetSettings, NativeCultureStatistics));
+		if (NativeCultureStatisticsPropertyHandle.IsValid() && NativeCultureStatisticsPropertyHandle->IsValidHandle())
+		{
+			NativeCultureNamePropertyHandle = NativeCultureStatisticsPropertyHandle->GetChildHandle(GET_MEMBER_NAME_CHECKED(FCultureStatistics, CultureName));
+			if (NativeCultureNamePropertyHandle.IsValid() && NativeCultureNamePropertyHandle->IsValidHandle())
+			{
+				NativeCultureNamePropertyHandle->MarkHiddenByCustomization();
+				FDetailWidgetRow& NativeCultureRow = DetailCategoryBuilder.AddCustomRow( LOCTEXT("NativeCultureNameFilterString", "Native Culture Name") );
+
+				FString NativeCultureName;
+				NativeCultureNamePropertyHandle->GetValue(NativeCultureName);
+				FCulturePtr NativeCulture = FInternationalization::Get().GetCulture(NativeCultureName);
+
+				NativeCultureRow.NameContent()
+					[
+						NativeCultureNamePropertyHandle->CreatePropertyNameWidget(LOCTEXT("NativeCultureNameLabel", "Native Culture"))
+					];
+
+				NativeCultureRow.ValueContent()
+					[
+						SAssignNew(NativeCultureComboButton, SComboButton)
+						.ButtonContent()
+						[
+							SNew(STextBlock)
+							.Text(this, &FLocalizationTargetDetailCustomization::GetNativeCultureDisplayName)
+							.ToolTipText(this, &FLocalizationTargetDetailCustomization::GetNativeCultureName)
+						]
+						.HasDownArrow(true)
+							.MenuContent()
+							[
+								SNew(SBox)
+								.MaxDesiredHeight(400.0f)
+								.MaxDesiredWidth(300.0f)
+								[
+									SNew(SCulturePicker)
+									.OnSelectionChanged(this, &FLocalizationTargetDetailCustomization::OnNativeCultureSelected)
+									.InitialSelection(NativeCulture)
+								]
+							]
+					];
+			}
+		}
+
+		SupportedCulturesStatisticsPropertyHandle = TargetSettingsPropertyHandle->GetChildHandle(GET_MEMBER_NAME_CHECKED(FLocalizationTargetSettings, SupportedCulturesStatistics));
+		if (SupportedCulturesStatisticsPropertyHandle.IsValid() && SupportedCulturesStatisticsPropertyHandle->IsValidHandle())
+		{
+			SupportedCulturesStatisticsPropertyHandle->MarkHiddenByCustomization();
+			SupportedCulturesStatisticsPropertyHandle_OnNumElementsChanged = FSimpleDelegate::CreateSP(this, &FLocalizationTargetDetailCustomization::RebuildListedCulturesList);
+			SupportedCulturesStatisticsPropertyHandle->AsArray()->SetOnNumElementsChanged(SupportedCulturesStatisticsPropertyHandle_OnNumElementsChanged);
+
+			FLocalizationTargetEditorCommands::Register();
+			const TSharedRef< FUICommandList > CommandList = MakeShareable(new FUICommandList);
+			FToolBarBuilder ToolBarBuilder( CommandList, FMultiBoxCustomization::AllowCustomization("LocalizationTargetEditor") );
+
+			CommandList->MapAction( FLocalizationTargetEditorCommands::Get().Gather, FExecuteAction::CreateSP(this, &FLocalizationTargetDetailCustomization::Gather));
+			ToolBarBuilder.AddToolBarButton(FLocalizationTargetEditorCommands::Get().Gather, NAME_None, TAttribute<FText>(), TAttribute<FText>(), FSlateIcon(FEditorStyle::GetStyleSetName(), "LocalizationTargetEditor.Gather"));
+
+			CommandList->MapAction( FLocalizationTargetEditorCommands::Get().ImportAllCultures, FExecuteAction::CreateSP(this, &FLocalizationTargetDetailCustomization::ImportAllCultures));
+			ToolBarBuilder.AddToolBarButton(FLocalizationTargetEditorCommands::Get().ImportAllCultures, NAME_None, TAttribute<FText>(), TAttribute<FText>(), FSlateIcon(FEditorStyle::GetStyleSetName(), "LocalizationTargetEditor.ImportForAllCultures"));
+
+			CommandList->MapAction( FLocalizationTargetEditorCommands::Get().ExportAllCultures, FExecuteAction::CreateSP(this, &FLocalizationTargetDetailCustomization::ExportAllCultures));
+			ToolBarBuilder.AddToolBarButton(FLocalizationTargetEditorCommands::Get().ExportAllCultures, NAME_None, TAttribute<FText>(), TAttribute<FText>(), FSlateIcon(FEditorStyle::GetStyleSetName(), "LocalizationTargetEditor.ExportForAllCultures"));
+
+			CommandList->MapAction( FLocalizationTargetEditorCommands::Get().RefreshWordCounts, FExecuteAction::CreateSP(this, &FLocalizationTargetDetailCustomization::RefreshWordCounts));
+			ToolBarBuilder.AddToolBarButton(FLocalizationTargetEditorCommands::Get().RefreshWordCounts, NAME_None, TAttribute<FText>(), TAttribute<FText>(), FSlateIcon(FEditorStyle::GetStyleSetName(), "LocalizationTargetEditor.RefreshWordCounts"));
+
+			CommandList->MapAction( FLocalizationTargetEditorCommands::Get().Compile, FExecuteAction::CreateSP(this, &FLocalizationTargetDetailCustomization::Compile));
+			ToolBarBuilder.AddToolBarButton(FLocalizationTargetEditorCommands::Get().Compile, NAME_None, TAttribute<FText>(), TAttribute<FText>(), FSlateIcon(FEditorStyle::GetStyleSetName(), "LocalizationTargetEditor.Compile"));
+
+
+			BuildListedCulturesList();
+
+			DetailCategoryBuilder.AddCustomRow( LOCTEXT("SupportedCultureNamesFilterString", "Supported Culture Names") )
+				.WholeRowContent()
+				[
+					SNew(SVerticalBox)
+					+SVerticalBox::Slot()
+					.AutoHeight()
+					[
+						ToolBarBuilder.MakeWidget()
+					]
+					+SVerticalBox::Slot()
+						.AutoHeight()
+						[
+							SAssignNew(SupportedCultureListView, SListView< TSharedPtr<IPropertyHandle> >)
+							.OnGenerateRow(this, &FLocalizationTargetDetailCustomization::OnGenerateCultureRow)
+							.ListItemsSource(&ListedCultureStatisticProperties)
+							.SelectionMode(ESelectionMode::None)
+							.HeaderRow
+							(
+							SNew(SHeaderRow)
+							+SHeaderRow::Column("Culture")
+							.DefaultLabel( NSLOCTEXT("LocalizationCulture", "CultureColumnLabel", "Culture"))
+							.HAlignHeader(HAlign_Fill)
+							.HAlignCell(HAlign_Fill)
+							.VAlignCell(VAlign_Center)
+							+SHeaderRow::Column("WordCount")
+							.DefaultLabel( NSLOCTEXT("LocalizationCulture", "WordCountColumnLabel", "Word Count"))
+							.HAlignHeader(HAlign_Center)
+							.HAlignCell(HAlign_Fill)
+							.VAlignCell(VAlign_Center)
+							+SHeaderRow::Column("Actions")
+							.DefaultLabel( NSLOCTEXT("LocalizationCulture", "ActionsColumnLabel", "Actions"))
+							.HAlignHeader(HAlign_Center)
+							.HAlignCell(HAlign_Center)
+							.VAlignCell(VAlign_Center)
+							)
+						]
+					+SVerticalBox::Slot()
+						.AutoHeight()
+						.VAlign(VAlign_Center)
+						[
+							SAssignNew(AddNewSupportedCultureComboButton, SComboButton)
+							.ButtonContent()
+							[
+								SNew(STextBlock)
+								.Text(NSLOCTEXT("LocalizationCulture", "AddNewCultureButtonLabel", "Add New Culture"))
+							]
+							.MenuContent()
+								[
+									SNew(SBox)
+									.MaxDesiredHeight(400.0f)
+									.MaxDesiredWidth(300.0f)
+									[
+										SAssignNew(SupportedCulturePicker, SCulturePicker)
+										.OnSelectionChanged(this, &FLocalizationTargetDetailCustomization::OnNewSupportedCultureSelected)
+										.IsCulturePickable(this, &FLocalizationTargetDetailCustomization::IsCultureSelectableAsSupported)
+									]
+								]
+						]
+				];
+		}
+	}
+}
+
+FLocalizationTargetSettings* FLocalizationTargetDetailCustomization::GetTargetSettings() const
+{
+	return LocalizationTarget.IsValid() ? &(LocalizationTarget->Settings) : nullptr;
+}
+
+TSharedPtr<IPropertyHandle> FLocalizationTargetDetailCustomization::GetTargetSettingsPropertyHandle() const
+{
+	return TargetSettingsPropertyHandle;
+}
+
+FText FLocalizationTargetDetailCustomization::GetTargetName() const
+{
+	if (LocalizationTarget.IsValid())
+	{
+		return FText::FromString(LocalizationTarget->Settings.Name);
+	}
+	return FText::GetEmpty();
+}
+
+bool FLocalizationTargetDetailCustomization::IsTargetNameUnique(const FString& Name) const
+{
+	for (ULocalizationTarget* const TargetObject : TargetSet->TargetObjects)
+	{
+		if (TargetObject != LocalizationTarget)
+		{
+			if (TargetObject->Settings.Name == LocalizationTarget->Settings.Name)
+			{
+				return false;
+			}
+		}
+	}
+	return true;
+}
+
+void FLocalizationTargetDetailCustomization::OnTargetNameChanged(const FText& NewText)
+{
+	const FString& NewName = NewText.ToString();
+
+	// TODO: Target names must be valid directory names, because they are used as directory names.
+	// ValidatePath allows /, which is not a valid directory name character
+	FText Error;
+	if (!FPaths::ValidatePath(NewName, &Error))
+	{
+		TargetNameEditableTextBox->SetError(Error);
+		return;		
+	}
+
+	// Target name must be unique.
+	if (!IsTargetNameUnique(NewName))
+	{
+		TargetNameEditableTextBox->SetError(LOCTEXT("DuplicateTargetNameError", "Target name must be unique."));
+		return;
+	}
+
+	// Clear error if nothing has failed.
+	TargetNameEditableTextBox->SetError(FText::GetEmpty());
+}
+
+void FLocalizationTargetDetailCustomization::OnTargetNameCommitted(const FText& NewText, ETextCommit::Type Type)
+{
+	// Target name must be unique.
+	if (!IsTargetNameUnique(NewText.ToString()))
+	{
+		return;
+	}
+
+	if (TargetSettingsPropertyHandle.IsValid() && TargetSettingsPropertyHandle->IsValidHandle())
+	{
+		FLocalizationTargetSettings* const TargetSettings = GetTargetSettings();
+		if (TargetSettings)
+		{
+			// Early out if the committed name is the same as the current name.
+			if (TargetSettings->Name == NewText.ToString())
+			{
+				return;
+			}
+
+			const TSharedPtr<IPropertyHandle> TargetNamePropertyHandle = TargetSettingsPropertyHandle->GetChildHandle(GET_MEMBER_NAME_CHECKED(FLocalizationTargetSettings, Name));
+
+			if (TargetNamePropertyHandle.IsValid() && TargetNamePropertyHandle->IsValidHandle())
+			{
+				TargetNamePropertyHandle->NotifyPreChange();
+			}
+
+			TargetSettings->Rename(NewText.ToString());
+
+			if (TargetNamePropertyHandle.IsValid() && TargetNamePropertyHandle->IsValidHandle())
+			{
+				TargetNamePropertyHandle->NotifyPostChange();
+			}
+		}
+	}
+}
+
+void FLocalizationTargetDetailCustomization::RebuildTargetDependenciesBox()
+{
+	if (TargetDependenciesHorizontalBox.IsValid())
+	{
+		for (const TSharedPtr<SWidget>& Widget : TargetDependenciesWidgets)
+		{
+			TargetDependenciesHorizontalBox->RemoveSlot(Widget.ToSharedRef());
+		}
+		TargetDependenciesWidgets.Empty();
+
+		for (const FString& TargetDependency : LocalizationTarget->Settings.TargetDependencies)
+		{
+			const TSharedRef<SWidget> Widget = SNew(SBorder)
+				[
+					SNew(STextBlock)
+					.Text(FText::FromString(TargetDependency))
+				];
+
+			TargetDependenciesWidgets.Add(Widget);
+			TargetDependenciesHorizontalBox->AddSlot()
+				[
+					Widget
+				];
+		}
+	}
+}
+
+void FLocalizationTargetDetailCustomization::RebuildTargetsList()
+{
+	TargetDependenciesOptionsList.Empty();
+	TFunction<bool (ULocalizationTarget* const)> DoesTargetDependOnUs;
+	DoesTargetDependOnUs = [&, this](ULocalizationTarget* const OtherTarget) -> bool
+	{
+		if (OtherTarget->Settings.TargetDependencies.Contains(LocalizationTarget->Settings.Name))
+		{
+			return true;
+		}
+
+		for (const FString& OtherTargetDependencyName : OtherTarget->Settings.TargetDependencies)
+		{
+			ULocalizationTarget** const OtherTargetDependency = TargetSet->TargetObjects.FindByPredicate([OtherTargetDependencyName](ULocalizationTarget* SomeLocalizationTarget)->bool{return SomeLocalizationTarget->Settings.Name == OtherTargetDependencyName;});
+			if (OtherTargetDependency && DoesTargetDependOnUs(*OtherTargetDependency))
+			{
+				return true;
+			}
+		}
+
+		return false;
+	};
+
+	for (ULocalizationTarget* const OtherTarget : TargetSet->TargetObjects)
+	{
+		if (OtherTarget != LocalizationTarget)
+		{
+			if (!DoesTargetDependOnUs(OtherTarget))
+			{
+				TargetDependenciesOptionsList.Add(OtherTarget);
+			}
+		}
+	}
+
+	if (TargetDependenciesListView.IsValid())
+	{
+		TargetDependenciesListView->RequestListRefresh();
+	}
+}
+
+TSharedRef<ITableRow> FLocalizationTargetDetailCustomization::OnGenerateTargetRow(ULocalizationTarget* OtherLocalizationTarget, const TSharedRef<STableViewBase>& Table)
+{
+	return SNew(STableRow<ULocalizationTarget*>, Table)
+		.ShowSelection(true)
+		.Content()
+		[
+			SNew(SCheckBox)
+			.OnCheckStateChanged_Lambda([this, OtherLocalizationTarget](ECheckBoxState State)
+			{
+				OnTargetDependencyCheckStateChanged(OtherLocalizationTarget, State);
+			})
+				.IsChecked_Lambda([this, OtherLocalizationTarget]()->ECheckBoxState
+			{
+				return IsTargetDependencyChecked(OtherLocalizationTarget);
+			})
+				.Content()
+				[
+					SNew(STextBlock)
+					.Text(FText::FromString(OtherLocalizationTarget->Settings.Name))
+				]
+		];
+}
+
+void FLocalizationTargetDetailCustomization::OnTargetDependencyCheckStateChanged(ULocalizationTarget* const OtherLocalizationTarget, const ECheckBoxState State)
+{
+	const TSharedPtr<IPropertyHandle> TargetDependenciesPropertyHandle = TargetSettingsPropertyHandle->GetChildHandle(GET_MEMBER_NAME_CHECKED(FLocalizationTargetSettings, TargetDependencies));
+
+	if (TargetDependenciesPropertyHandle.IsValid() && TargetDependenciesPropertyHandle->IsValidHandle())
+	{
+		TargetDependenciesPropertyHandle->NotifyPreChange();
+	}
+
+	switch (State)
+	{
+	case ECheckBoxState::Checked:
+		LocalizationTarget->Settings.TargetDependencies.Add(OtherLocalizationTarget->Settings.Name);
+		break;
+	case ECheckBoxState::Unchecked:
+		LocalizationTarget->Settings.TargetDependencies.Remove(OtherLocalizationTarget->Settings.Name);
+		break;
+	}
+
+	if (TargetDependenciesPropertyHandle.IsValid() && TargetDependenciesPropertyHandle->IsValidHandle())
+	{
+		TargetDependenciesPropertyHandle->NotifyPostChange();
+	}
+
+	RebuildTargetDependenciesBox();
+}
+
+ECheckBoxState FLocalizationTargetDetailCustomization::IsTargetDependencyChecked(ULocalizationTarget* const OtherLocalizationTarget) const
+{
+	return LocalizationTarget->Settings.TargetDependencies.Contains(OtherLocalizationTarget->Settings.Name) ? ECheckBoxState::Checked : ECheckBoxState::Unchecked;
+}
+
+FText FLocalizationTargetDetailCustomization::GetNativeCultureName() const
+{
+	FCulturePtr NativeCulture;
+
+	if(NativeCultureNamePropertyHandle.IsValid() && NativeCultureNamePropertyHandle->IsValidHandle())
+	{
+		FString Value;
+		NativeCultureNamePropertyHandle->GetValue(Value);
+		NativeCulture = FInternationalization::Get().GetCulture(Value);
+	}
+
+	return NativeCulture.IsValid() ? FText::FromString(NativeCulture->GetName()) : FText::GetEmpty();
+}
+
+FText FLocalizationTargetDetailCustomization::GetNativeCultureDisplayName() const
+{
+	FCulturePtr NativeCulture;
+
+	if(NativeCultureNamePropertyHandle.IsValid() && NativeCultureNamePropertyHandle->IsValidHandle())
+	{
+		FString Value;
+		NativeCultureNamePropertyHandle->GetValue(Value);
+		NativeCulture = FInternationalization::Get().GetCulture(Value);
+	}
+
+	return NativeCulture.IsValid() ? FText::FromString(NativeCulture->GetDisplayName()) : FText::GetEmpty();
+}
+
+void FLocalizationTargetDetailCustomization::OnNativeCultureSelected(FCulturePtr SelectedCulture, ESelectInfo::Type SelectInfo)
+{
+	if(NativeCultureNamePropertyHandle.IsValid() && NativeCultureNamePropertyHandle->IsValidHandle())
+	{
+		NativeCultureNamePropertyHandle->SetValue(SelectedCulture->GetName());
+	}
+	if (NativeCultureComboButton.IsValid())
+	{
+		NativeCultureComboButton->SetIsOpen(false);
+	}
+
+	// TODO:	If the selected culture is an existing supported culture:
+	//			copy that supported culture's statistics to the native culture statistics
+	//			remove the supported culture's statistics
+
+	// TODO: Warn about this action invalidating all word counts and, realistically, invalidating all existing translations.
+	// Possibly offer an option to delete all archives?
+}
+
+void FLocalizationTargetDetailCustomization::Gather()
+{
+	if (LocalizationTarget.IsValid())
+	{
+		// Save unsaved packages.
+		const bool bPromptUserToSave = true;
+		const bool bSaveMapPackages = true;
+		const bool bSaveContentPackages = true;
+		const bool bFastSave = false;
+		const bool bNotifyNoPackagesSaved = false;
+		const bool bCanBeDeclined = true;
+		bool DidPackagesNeedSaving;
+		const bool WerePackagesSaved = FEditorFileUtils::SaveDirtyPackages(bPromptUserToSave, bSaveMapPackages, bSaveContentPackages, bFastSave, bNotifyNoPackagesSaved, bCanBeDeclined, &DidPackagesNeedSaving);
+
+		if (DidPackagesNeedSaving && !WerePackagesSaved)
+		{
+			// Give warning dialog.
+			const FText MessageText = NSLOCTEXT("LocalizationCultureActions", "UnsavedPackagesWarningDialogMessage", "There are unsaved changes. These changes may not be gathered from correctly.");
+			const FText TitleText = NSLOCTEXT("LocalizationCultureActions", "UnsavedPackagesWarningDialogTitle", "Unsaved Changes Before Gather");
+			switch(FMessageDialog::Open(EAppMsgType::OkCancel, MessageText, &TitleText))
+			{
+			case EAppReturnType::Cancel:
+				{
+					return;
+				}
+				break;
+			}
+		}
+
+		// Execute gather.
+		const TSharedPtr<SWindow> ParentWindow = FSlateApplication::Get().FindWidgetWindow(DetailLayoutBuilder->GetDetailsView().AsShared());
+		LocalizationCommandletTasks::GatherTarget(ParentWindow.ToSharedRef(), LocalizationTarget->Settings);
+
+		UpdateTargetFromReports();
+	}
+}
+
+void FLocalizationTargetDetailCustomization::ImportAllCultures()
+{
+	IDesktopPlatform* DesktopPlatform = FDesktopPlatformModule::Get();
+	if (LocalizationTarget.IsValid() && DesktopPlatform)
+	{
+		void* ParentWindowWindowHandle = NULL;
+		const TSharedPtr<SWindow> ParentWindow = FSlateApplication::Get().FindWidgetWindow(DetailLayoutBuilder->GetDetailsView().AsShared());
+		if (ParentWindow.IsValid() && ParentWindow->GetNativeWindow().IsValid())
+		{
+			ParentWindowWindowHandle = ParentWindow->GetNativeWindow()->GetOSWindowHandle();
+		}
+
+		const FString DefaultPath = FPaths::ConvertRelativePathToFull(LocalizationConfigurationScript::GetDataDirectory(LocalizationTarget->Settings));
+
+		FText DialogTitle;
+		{
+			FFormatNamedArguments FormatArguments;
+			FormatArguments.Add(TEXT("TargetName"), FText::FromString(LocalizationTarget->Settings.Name));
+			DialogTitle = FText::Format(LOCTEXT("ImportAllTranslationsForTargetDialogTitleFormat", "Import All Translations for {TargetName} from Directory"), FormatArguments);
+		}
+
+		// Prompt the user for the directory
+		FString OutputDirectory;
+		if (DesktopPlatform->OpenDirectoryDialog(ParentWindowWindowHandle, DialogTitle.ToString(), DefaultPath, OutputDirectory))
+		{
+			LocalizationCommandletTasks::ImportTarget(ParentWindow.ToSharedRef(), LocalizationTarget->Settings, TOptional<FString>(OutputDirectory));
+
+			UpdateTargetFromReports();
+		}
+	}
+}
+
+void FLocalizationTargetDetailCustomization::ExportAllCultures()
+{
+	IDesktopPlatform* DesktopPlatform = FDesktopPlatformModule::Get();
+	if (LocalizationTarget.IsValid() && DesktopPlatform)
+	{
+		void* ParentWindowWindowHandle = NULL;
+		const TSharedPtr<SWindow> ParentWindow = FSlateApplication::Get().FindWidgetWindow(DetailLayoutBuilder->GetDetailsView().AsShared());
+		if (ParentWindow.IsValid() && ParentWindow->GetNativeWindow().IsValid())
+		{
+			ParentWindowWindowHandle = ParentWindow->GetNativeWindow()->GetOSWindowHandle();
+		}
+
+		const FString DefaultPath = FPaths::ConvertRelativePathToFull(LocalizationConfigurationScript::GetDataDirectory(LocalizationTarget->Settings));
+
+		FText DialogTitle;
+		{
+			FFormatNamedArguments FormatArguments;
+			FormatArguments.Add(TEXT("TargetName"), FText::FromString(LocalizationTarget->Settings.Name));
+			DialogTitle = FText::Format(LOCTEXT("ExportAllTranslationsForTargetDialogTitleFormat", "Export All Translations for {TargetName} to Directory"), FormatArguments);
+		}
+
+		// Prompt the user for the directory
+		FString OutputDirectory;
+		if (DesktopPlatform->OpenDirectoryDialog(ParentWindowWindowHandle, DialogTitle.ToString(), DefaultPath, OutputDirectory))
+		{
+			LocalizationCommandletTasks::ExportTarget(ParentWindow.ToSharedRef(), LocalizationTarget->Settings, TOptional<FString>(OutputDirectory));
+		}
+	}
+}
+
+void FLocalizationTargetDetailCustomization::Compile()
+{
+	if (LocalizationTarget.IsValid())
+	{
+		// Execute compile.
+		const TSharedPtr<SWindow> ParentWindow = FSlateApplication::Get().FindWidgetWindow(DetailLayoutBuilder->GetDetailsView().AsShared());
+		LocalizationCommandletTasks::CompileTarget(ParentWindow.ToSharedRef(), LocalizationTarget->Settings);
+	}
+}
+
+void FLocalizationTargetDetailCustomization::RefreshWordCounts()
+{
+	if (LocalizationTarget.IsValid())
+	{
+		const TSharedPtr<SWindow> ParentWindow = FSlateApplication::Get().FindWidgetWindow(DetailLayoutBuilder->GetDetailsView().AsShared());
+		LocalizationCommandletTasks::GenerateReportsForTarget(ParentWindow.ToSharedRef(), LocalizationTarget->Settings);
+
+		UpdateTargetFromReports();
+	}
+}
+
+void FLocalizationTargetDetailCustomization::UpdateTargetFromReports()
+{
+	FLocalizationTargetSettings* const TargetSettings = GetTargetSettings();
+	if (TargetSettings)
+	{
+		TArray< TSharedPtr<IPropertyHandle> > WordCountPropertyHandles;
+
+		if (TargetSettingsPropertyHandle.IsValid() && TargetSettingsPropertyHandle->IsValidHandle())
+		{
+			const TSharedPtr<IPropertyHandle> NativeCultureStatisticsPropertyHandle = TargetSettingsPropertyHandle->GetChildHandle(GET_MEMBER_NAME_CHECKED(FLocalizationTargetSettings, NativeCultureStatistics));
+			if (NativeCultureStatisticsPropertyHandle.IsValid() && NativeCultureStatisticsPropertyHandle->IsValidHandle())
+			{
+				const TSharedPtr<IPropertyHandle> WordCountPropertyHandle = NativeCultureStatisticsPropertyHandle->GetChildHandle(GET_MEMBER_NAME_CHECKED(FCultureStatistics, WordCount));
+				if (WordCountPropertyHandle.IsValid() && WordCountPropertyHandle->IsValidHandle())
+				{
+					WordCountPropertyHandles.Add(WordCountPropertyHandle);
+				}
+			}
+
+			const TSharedPtr<IPropertyHandle> SupportedCulturesStatisticsPropertyHandle = TargetSettingsPropertyHandle->GetChildHandle(GET_MEMBER_NAME_CHECKED(FLocalizationTargetSettings, SupportedCulturesStatistics));
+			if (SupportedCulturesStatisticsPropertyHandle.IsValid() && SupportedCulturesStatisticsPropertyHandle->IsValidHandle())
+			{
+				uint32 SupportedCultureCount = 0;
+				SupportedCulturesStatisticsPropertyHandle->GetNumChildren(SupportedCultureCount);
+				for (uint32 i = 0; i < SupportedCultureCount; ++i)
+				{
+					const TSharedPtr<IPropertyHandle> ElementPropertyHandle = SupportedCulturesStatisticsPropertyHandle->GetChildHandle(i);
+					if (ElementPropertyHandle.IsValid() && ElementPropertyHandle->IsValidHandle())
+					{
+						const TSharedPtr<IPropertyHandle> WordCountPropertyHandle = SupportedCulturesStatisticsPropertyHandle->GetChildHandle(GET_MEMBER_NAME_CHECKED(FCultureStatistics, WordCount));
+						if (WordCountPropertyHandle.IsValid() && WordCountPropertyHandle->IsValidHandle())
+						{
+							WordCountPropertyHandles.Add(WordCountPropertyHandle);
+						}
+					}
+				}
+			}
+		}
+
+		for (const TSharedPtr<IPropertyHandle>& WordCountPropertyHandle : WordCountPropertyHandles)
+		{
+			WordCountPropertyHandle->NotifyPreChange();
+		}
+		TargetSettings->UpdateWordCountsFromCSV();
+		TargetSettings->UpdateStatusFromConflictReport();
+		for (const TSharedPtr<IPropertyHandle>& WordCountPropertyHandle : WordCountPropertyHandles)
+		{
+			WordCountPropertyHandle->NotifyPostChange();
+		}
+	}
+}
+
+void FLocalizationTargetDetailCustomization::BuildListedCulturesList()
+{
+	ListedCultureStatisticProperties.AddUnique(NativeCultureStatisticsPropertyHandle);
+	const TSharedPtr<IPropertyHandleArray> SupportedCulturesStatisticsArrayPropertyHandle = SupportedCulturesStatisticsPropertyHandle->AsArray();
+	if (SupportedCulturesStatisticsArrayPropertyHandle.IsValid())
+	{
+		uint32 ElementCount = 0;
+		SupportedCulturesStatisticsArrayPropertyHandle->GetNumElements(ElementCount);
+		for (uint32 i = 0; i < ElementCount; ++i)
+		{
+			const TSharedPtr<IPropertyHandle> CultureStatisticsProperty = SupportedCulturesStatisticsArrayPropertyHandle->GetElement(i);
+			ListedCultureStatisticProperties.AddUnique(CultureStatisticsProperty);
+		}
+	}
+}
+
+void FLocalizationTargetDetailCustomization::RebuildListedCulturesList()
+{
+	if (NewEntryIndexToBeInitialized != INDEX_NONE)
+	{
+		const TSharedPtr<IPropertyHandle> SupportedCultureStatisticsPropertyHandle = SupportedCulturesStatisticsPropertyHandle->GetChildHandle(NewEntryIndexToBeInitialized);
+
+		const TSharedPtr<IPropertyHandle> CultureNamePropertyHandle = SupportedCultureStatisticsPropertyHandle->GetChildHandle(GET_MEMBER_NAME_CHECKED(FCultureStatistics, CultureName));
+		if(CultureNamePropertyHandle.IsValid() && CultureNamePropertyHandle->IsValidHandle())
+		{
+			CultureNamePropertyHandle->SetValue(SelectedNewCulture->GetName());
+		}
+
+		const TSharedPtr<IPropertyHandle> WordCountPropertyHandle = SupportedCultureStatisticsPropertyHandle->GetChildHandle(GET_MEMBER_NAME_CHECKED(FCultureStatistics, WordCount));
+		if(WordCountPropertyHandle.IsValid() && WordCountPropertyHandle->IsValidHandle())
+		{
+			WordCountPropertyHandle->SetValue(0);
+		}
+
+		AddNewSupportedCultureComboButton->SetIsOpen(false);
+
+		NewEntryIndexToBeInitialized = INDEX_NONE;
+	}
+
+	ListedCultureStatisticProperties.Empty();
+	BuildListedCulturesList();
+
+	if (SupportedCultureListView.IsValid())
+	{
+		SupportedCultureListView->RequestListRefresh();
+	}
+}
+
+TSharedRef<ITableRow> FLocalizationTargetDetailCustomization::OnGenerateCultureRow(TSharedPtr<IPropertyHandle> CultureStatisticsPropertyHandle, const TSharedRef<STableViewBase>& Table)
+{
+	return SNew(SLocalizationTargetEditorCultureRow, Table, DetailLayoutBuilder->GetPropertyUtilities(), TargetSettingsPropertyHandle.ToSharedRef(), CultureStatisticsPropertyHandle->GetIndexInArray());
+}
+
+bool FLocalizationTargetDetailCustomization::IsCultureSelectableAsSupported(FCulturePtr Culture)
+{
+	// Can't select the native culture.
+	if (NativeCultureNamePropertyHandle.IsValid() && NativeCultureNamePropertyHandle->IsValidHandle())
+	{
+		FString NativeCultureName;
+		NativeCultureNamePropertyHandle->GetValue(NativeCultureName);
+		const FCulturePtr NativeCulture = FInternationalization::Get().GetCulture(NativeCultureName);
+		if (NativeCulture == Culture)
+		{
+			return false;
+		}
+	}
+
+	auto Is = [&](const TSharedPtr<IPropertyHandle>& SupportedCultureStatisticProperty)
+	{
+		// Can't select existing supported cultures.
+		const TSharedPtr<IPropertyHandle> SupportedCultureNamePropertyHandle = SupportedCultureStatisticProperty->GetChildHandle(GET_MEMBER_NAME_CHECKED(FCultureStatistics, CultureName));
+		if (SupportedCultureNamePropertyHandle.IsValid() && SupportedCultureNamePropertyHandle->IsValidHandle())
+		{
+			FString SupportedCultureName;
+			SupportedCultureNamePropertyHandle->GetValue(SupportedCultureName);
+			const FCulturePtr SupportedCulture = FInternationalization::Get().GetCulture(SupportedCultureName);
+			return SupportedCulture == Culture;
+		}
+
+		return false;
+	};
+
+	return !ListedCultureStatisticProperties.ContainsByPredicate(Is);
+}
+
+void FLocalizationTargetDetailCustomization::OnNewSupportedCultureSelected(FCulturePtr SelectedCulture, ESelectInfo::Type SelectInfo)
+{
+	if(SupportedCulturesStatisticsPropertyHandle.IsValid() && SupportedCulturesStatisticsPropertyHandle->IsValidHandle())
+	{
+		uint32 NewElementIndex;
+		SupportedCulturesStatisticsPropertyHandle->AsArray()->GetNumElements(NewElementIndex);
+
+		// Add element, set info for later initialization.
+		SupportedCulturesStatisticsPropertyHandle->AsArray()->AddItem();
+		SelectedNewCulture = SelectedCulture;
+		NewEntryIndexToBeInitialized = NewElementIndex;
+
+		// Refresh UI.
+		if (SupportedCulturePicker.IsValid())
+		{
+			SupportedCulturePicker->RequestTreeRefresh();
+		}
+	}
+}
+
+#undef LOCTEXT_NAMESPACE
\ No newline at end of file
diff --git a/Engine/Source/Editor/LocalizationDashboard/Private/LocalizationTargetDetailCustomization.h b/Engine/Source/Editor/LocalizationDashboard/Private/LocalizationTargetDetailCustomization.h
new file mode 100644
--- /dev/null
+++ b/Engine/Source/Editor/LocalizationDashboard/Private/LocalizationTargetDetailCustomization.h
@@ -0,0 +1,79 @@
+// Copyright 1998-2015 Epic Games, Inc. All Rights Reserved.
+
+#pragma once
+#include "IDetailCustomization.h"
+
+class ULocalizationTargetSet;
+class ULocalizationTarget;
+class SCulturePicker;
+
+class FLocalizationTargetDetailCustomization : public IDetailCustomization
+{
+public:
+	FLocalizationTargetDetailCustomization();
+	virtual void CustomizeDetails(IDetailLayoutBuilder& DetailBuilder) override;
+
+	FLocalizationTargetSettings* GetTargetSettings() const;
+	TSharedPtr<IPropertyHandle> GetTargetSettingsPropertyHandle() const;
+
+private:
+	FText GetTargetName() const;
+	bool IsTargetNameUnique(const FString& Name) const;
+	void OnTargetNameChanged(const FText& NewText);
+	void OnTargetNameCommitted(const FText& NewText, ETextCommit::Type Type);
+
+	void RebuildTargetDependenciesBox();
+	void RebuildTargetsList();
+	TSharedRef<ITableRow> OnGenerateTargetRow(ULocalizationTarget* OtherLocalizationTarget, const TSharedRef<STableViewBase>& Table);
+	void OnTargetDependencyCheckStateChanged(ULocalizationTarget* const OtherLocalizationTarget, const ECheckBoxState State);
+	ECheckBoxState IsTargetDependencyChecked(ULocalizationTarget* const OtherLocalizationTarget) const;
+
+	FText GetNativeCultureName() const;
+	FText GetNativeCultureDisplayName() const;
+	void OnNativeCultureSelected(FCulturePtr SelectedCulture, ESelectInfo::Type SelectInfo);
+
+	void Gather();
+	void ImportAllCultures();
+	void ExportAllCultures();
+	void Compile();
+	void RefreshWordCounts();
+	void UpdateTargetFromReports();
+
+	void BuildListedCulturesList();
+	void RebuildListedCulturesList();
+	TSharedRef<ITableRow> OnGenerateCultureRow(TSharedPtr<IPropertyHandle> CulturePropertyHandle, const TSharedRef<STableViewBase>& Table);
+
+	bool IsCultureSelectableAsSupported(FCulturePtr Culture);
+	void OnNewSupportedCultureSelected(FCulturePtr SelectedCulture, ESelectInfo::Type SelectInfo);
+
+private:
+	IDetailLayoutBuilder* DetailLayoutBuilder;
+
+	TWeakObjectPtr<ULocalizationTargetSet> TargetSet;
+	TWeakObjectPtr<ULocalizationTarget> LocalizationTarget;
+
+	TSharedPtr<IPropertyHandle> TargetSettingsPropertyHandle;
+
+	TSharedPtr<SEditableTextBox> TargetNameEditableTextBox;
+
+	TSharedPtr<SHorizontalBox> TargetDependenciesHorizontalBox;
+	TArray< TSharedPtr<SWidget> > TargetDependenciesWidgets;
+	TArray<ULocalizationTarget*> TargetDependenciesOptionsList;
+	TSharedPtr< SListView<ULocalizationTarget*> > TargetDependenciesListView;
+
+	TSharedPtr<IPropertyHandle> NativeCultureStatisticsPropertyHandle;
+	TSharedPtr<IPropertyHandle> NativeCultureNamePropertyHandle;
+	TSharedPtr<SComboButton> NativeCultureComboButton;
+	TArray<FCulturePtr> AllCultures;
+
+	TSharedPtr<IPropertyHandle> SupportedCulturesStatisticsPropertyHandle;
+	FSimpleDelegate SupportedCulturesStatisticsPropertyHandle_OnNumElementsChanged;
+	TSharedPtr< SListView< TSharedPtr<IPropertyHandle> > > SupportedCultureListView;
+	TSharedPtr<SComboButton> AddNewSupportedCultureComboButton;
+	TSharedPtr<SCulturePicker> SupportedCulturePicker;
+	TArray< TSharedPtr<IPropertyHandle> > ListedCultureStatisticProperties;
+
+	/* If set, the entry at the index specified needs to be initialized as soon as possible. */
+	int32 NewEntryIndexToBeInitialized;
+	FCulturePtr SelectedNewCulture;
+};
diff --git a/Engine/Source/Editor/LocalizationDashboard/Private/LocalizationTargetSetDetailCustomization.cpp b/Engine/Source/Editor/LocalizationDashboard/Private/LocalizationTargetSetDetailCustomization.cpp
new file mode 100644
--- /dev/null
+++ b/Engine/Source/Editor/LocalizationDashboard/Private/LocalizationTargetSetDetailCustomization.cpp
@@ -0,0 +1,438 @@
+// Copyright 1998-2015 Epic Games, Inc. All Rights Reserved.
+
+#include "LocalizationDashboardPrivatePCH.h"
+#include "LocalizationTargetSetDetailCustomization.h"
+#include "DetailLayoutBuilder.h"
+#include "LocalizationTarget.h"
+#include "DetailWidgetRow.h"
+#include "SLocalizationDashboardTargetRow.h"
+
+#define LOCTEXT_NAMESPACE "LocalizationDashboard"
+
+FLocalizationTargetSetDetailCustomization::FLocalizationTargetSetDetailCustomization()
+	: DetailLayoutBuilder(nullptr)
+	//, ServiceProviderCategoryBuilder(nullptr)
+	, NewEntryIndexToBeInitialized(INDEX_NONE)
+{
+	//TArray<ILocalizationServiceProvider*> ActualProviders = ILocalizationDashboardModule::Get().GetLocalizationServiceProviders();
+	//for (ILocalizationServiceProvider* ActualProvider : ActualProviders)
+	//{
+	//	TSharedPtr<FLocalizationServiceProviderWrapper> Provider = MakeShareable(new FLocalizationServiceProviderWrapper(ActualProvider));
+	//	Providers.Add(Provider);
+	//}
+}
+
+//class FindProviderPredicate
+//{
+//public:
+//	FindProviderPredicate(ILocalizationServiceProvider* const InActualProvider)
+//		: ActualProvider(InActualProvider)
+//	{
+//	}
+
+//	bool operator()(const TSharedPtr<FLocalizationServiceProviderWrapper>& Provider)
+//	{
+//		return Provider->Provider == ActualProvider;
+//	}
+
+//private:
+//	ILocalizationServiceProvider* ActualProvider;
+//};
+
+namespace
+{
+	class FLocalizationDashboardCommands : public TCommands<FLocalizationDashboardCommands>
+	{
+	public:
+		FLocalizationDashboardCommands() 
+			: TCommands<FLocalizationDashboardCommands>("LocalizationDashboard", NSLOCTEXT("Contexts", "LocalizationDashboard", "Localization Dashboard"), NAME_None, FEditorStyle::GetStyleSetName())
+		{
+		}
+
+		TSharedPtr<FUICommandInfo> GatherAllTargets;
+		TSharedPtr<FUICommandInfo> ImportAllTargets;
+		TSharedPtr<FUICommandInfo> ExportAllTargets;
+		TSharedPtr<FUICommandInfo> CompileAllTargets;
+
+		/** Initialize commands */
+		virtual void RegisterCommands() override;
+	};
+
+	void FLocalizationDashboardCommands::RegisterCommands() 
+	{
+		UI_COMMAND( GatherAllTargets, "Gather All", "Gathers text for all targets in the project.", EUserInterfaceActionType::Button, FInputGesture() );
+		UI_COMMAND( ImportAllTargets, "Import All", "Imports translations for all cultures of all targets in the project.", EUserInterfaceActionType::Button, FInputGesture() );
+		UI_COMMAND( ExportAllTargets, "Export All", "Exports translations for all cultures of all targets in the project.", EUserInterfaceActionType::Button, FInputGesture() );
+		UI_COMMAND( CompileAllTargets, "Compile All", "Compiles translations for all targets in the project.", EUserInterfaceActionType::Button, FInputGesture() );
+	}
+}
+
+void FLocalizationTargetSetDetailCustomization::CustomizeDetails(IDetailLayoutBuilder& DetailBuilder)
+{
+	DetailLayoutBuilder = &DetailBuilder;
+	{
+		TArray< TWeakObjectPtr<UObject> > ObjectsBeingCustomized;
+		DetailLayoutBuilder->GetObjectsBeingCustomized(ObjectsBeingCustomized);
+		TargetSet = CastChecked<ULocalizationTargetSet>(ObjectsBeingCustomized.Top().Get());
+	}
+
+	{
+		IDetailCategoryBuilder& DetailCategoryBuilder = DetailBuilder.EditCategory(FName("Targets"));
+		TargetObjectsPropertyHandle = DetailBuilder.GetProperty(GET_MEMBER_NAME_CHECKED(ULocalizationTargetSet,TargetObjects));
+		if (TargetObjectsPropertyHandle->IsValidHandle())
+		{
+			TargetObjectsPropertyHandle->MarkHiddenByCustomization();
+			TargetsArrayPropertyHandle_OnNumElementsChanged = FSimpleDelegate::CreateSP(this, &FLocalizationTargetSetDetailCustomization::RebuildTargetsList);
+			TargetObjectsPropertyHandle->AsArray()->SetOnNumElementsChanged(TargetsArrayPropertyHandle_OnNumElementsChanged);
+
+			FLocalizationDashboardCommands::Register();
+			const TSharedRef< FUICommandList > CommandList = MakeShareable(new FUICommandList);
+			FToolBarBuilder ToolBarBuilder( CommandList, FMultiBoxCustomization::AllowCustomization("LocalizationDashboard") );
+
+			CommandList->MapAction( FLocalizationDashboardCommands::Get().GatherAllTargets, FExecuteAction::CreateSP(this, &FLocalizationTargetSetDetailCustomization::GatherAllTargets));
+			ToolBarBuilder.AddToolBarButton(FLocalizationDashboardCommands::Get().GatherAllTargets, NAME_None, TAttribute<FText>(), TAttribute<FText>(), FSlateIcon(FEditorStyle::GetStyleSetName(), "LocalizationDashboard.GatherAllTargets"));
+
+			CommandList->MapAction( FLocalizationDashboardCommands::Get().ImportAllTargets, FExecuteAction::CreateSP(this, &FLocalizationTargetSetDetailCustomization::ImportAllTargets));
+			ToolBarBuilder.AddToolBarButton(FLocalizationDashboardCommands::Get().ImportAllTargets, NAME_None, TAttribute<FText>(), TAttribute<FText>(), FSlateIcon(FEditorStyle::GetStyleSetName(), "LocalizationDashboard.ImportForAllTargetsCultures"));
+
+			CommandList->MapAction( FLocalizationDashboardCommands::Get().ExportAllTargets, FExecuteAction::CreateSP(this, &FLocalizationTargetSetDetailCustomization::ExportAllTargets));
+			ToolBarBuilder.AddToolBarButton(FLocalizationDashboardCommands::Get().ExportAllTargets, NAME_None, TAttribute<FText>(), TAttribute<FText>(), FSlateIcon(FEditorStyle::GetStyleSetName(), "LocalizationDashboard.ExportForAllTargetsCultures"));
+
+			CommandList->MapAction( FLocalizationDashboardCommands::Get().CompileAllTargets, FExecuteAction::CreateSP(this, &FLocalizationTargetSetDetailCustomization::CompileAllTargets));
+			ToolBarBuilder.AddToolBarButton(FLocalizationDashboardCommands::Get().CompileAllTargets, NAME_None, TAttribute<FText>(), TAttribute<FText>(), FSlateIcon(FEditorStyle::GetStyleSetName(), "LocalizationDashboard.CompileAllTargets"));
+
+			BuildTargetsList();
+
+			DetailCategoryBuilder.AddCustomRow( LOCTEXT("TargetsFilterString", "Targets") )
+				.WholeRowContent()
+				[
+					SNew(SVerticalBox)
+					+SVerticalBox::Slot()
+					.AutoHeight()
+					[
+						ToolBarBuilder.MakeWidget()
+					]
+					+SVerticalBox::Slot()
+						.AutoHeight()
+						[
+							SAssignNew(TargetsListView, SListView< TSharedPtr<IPropertyHandle> >)
+							.OnGenerateRow(this, &FLocalizationTargetSetDetailCustomization::OnGenerateRow)
+							.ListItemsSource(&TargetsList)
+							.SelectionMode(ESelectionMode::None)
+							.HeaderRow
+							(
+							SNew(SHeaderRow)
+							+SHeaderRow::Column("Target")
+							.DefaultLabel( LOCTEXT("TargetColumnLabel", "Target"))
+							.HAlignHeader(HAlign_Left)
+							.HAlignCell(HAlign_Left)
+							.VAlignCell(VAlign_Center)
+							+SHeaderRow::Column("Status")
+							.DefaultLabel( LOCTEXT("StatusColumnLabel", "Status"))
+							.HAlignHeader(HAlign_Center)
+							.HAlignCell(HAlign_Center)
+							.VAlignCell(VAlign_Center)
+							+SHeaderRow::Column("Cultures")
+							.DefaultLabel( LOCTEXT("CulturesColumnLabel", "Cultures"))
+							.HAlignHeader(HAlign_Fill)
+							.HAlignCell(HAlign_Fill)
+							.VAlignCell(VAlign_Center)
+							+SHeaderRow::Column("WordCount")
+							.DefaultLabel( LOCTEXT("WordCountColumnLabel", "Word Count"))
+							.HAlignHeader(HAlign_Center)
+							.HAlignCell(HAlign_Center)
+							.VAlignCell(VAlign_Center)
+							+SHeaderRow::Column("Actions")
+							.DefaultLabel( LOCTEXT("ActionsColumnLabel", "Actions"))
+							.HAlignHeader(HAlign_Center)
+							.HAlignCell(HAlign_Center)
+							.VAlignCell(VAlign_Center)
+							)
+						]
+					+SVerticalBox::Slot()
+						.AutoHeight()
+						.VAlign(VAlign_Center)
+						[
+							SNew(SButton)
+							.Text(LOCTEXT("AddNewTargetButtonLabel", "Add New Target"))
+							.OnClicked(this, &FLocalizationTargetSetDetailCustomization::OnNewTargetButtonClicked)
+						]
+				];
+		}
+	}
+}
+
+void FLocalizationTargetSetDetailCustomization::BuildTargetsList()
+{
+	const TSharedPtr<IPropertyHandleArray> TargetObjectsArrayPropertyHandle = TargetObjectsPropertyHandle->AsArray();
+	if (TargetObjectsArrayPropertyHandle.IsValid())
+	{
+		uint32 ElementCount = 0;
+		TargetObjectsArrayPropertyHandle->GetNumElements(ElementCount);
+		for (uint32 i = 0; i < ElementCount; ++i)
+		{
+			const TSharedPtr<IPropertyHandle> TargetObjectPropertyHandle = TargetObjectsArrayPropertyHandle->GetElement(i);
+			TargetsList.Add(TargetObjectPropertyHandle);
+		}
+	}
+}
+
+void FLocalizationTargetSetDetailCustomization::RebuildTargetsList()
+{
+	const TSharedPtr<IPropertyHandleArray> TargetObjectsArrayPropertyHandle = TargetObjectsPropertyHandle->AsArray();
+	if (TargetObjectsArrayPropertyHandle.IsValid() && NewEntryIndexToBeInitialized != INDEX_NONE)
+	{
+		const TSharedPtr<IPropertyHandle> TargetObjectPropertyHandle = TargetObjectsArrayPropertyHandle->GetElement(NewEntryIndexToBeInitialized);
+		if (TargetObjectPropertyHandle.IsValid() && TargetObjectPropertyHandle->IsValidHandle())
+		{
+			ULocalizationTarget* const NewTarget = NewObject<ULocalizationTarget>(TargetSet.Get());
+
+			TArray<void*> RawData;
+			TargetObjectsPropertyHandle->AccessRawData(RawData);
+			void* RawDatum = RawData.Top();
+			TArray<ULocalizationTarget*>& TargetObjectsArray = *(reinterpret_cast< TArray<ULocalizationTarget*>* >(RawDatum));
+
+			FString NewTargetName = "NewTarget";
+			auto TargetNameIsIdentical = [&](ULocalizationTarget* Target) -> bool
+			{
+				return (Target != NewTarget) && Target && (Target->Settings.Name == NewTargetName);
+			};
+
+			for (uint32 i = 0; TargetObjectsArray.ContainsByPredicate(TargetNameIsIdentical); ++i)
+			{
+				NewTargetName = FString::Printf(TEXT("NewTarget%u"), i);
+			}
+
+			NewTarget->Settings.Name = NewTargetName;
+			NewTarget->Settings.NativeCultureStatistics.CultureName = FInternationalization::Get().GetCurrentCulture()->GetName();
+
+			const UObject* SetValue = NewTarget;
+			TargetObjectPropertyHandle->SetValue(SetValue);
+
+			NewEntryIndexToBeInitialized = INDEX_NONE;
+		}
+	}
+
+	TargetsList.Empty();
+	BuildTargetsList();
+
+	if (TargetsListView.IsValid())
+	{
+		TargetsListView->RequestListRefresh();
+	}
+}
+
+//FText FLocalizationTargetSetDetailCustomization::GetCurrentServiceProviderDisplayName() const
+//{
+//	ILocalizationServiceProvider* const LSP = ILocalizationDashboardModule::Get().GetCurrentLocalizationServiceProvider();
+//	return LSP ? LSP->GetDisplayName() : LOCTEXT("NoServiceProviderName", "None");
+//}
+//
+//TSharedRef<SWidget> FLocalizationTargetSetDetailCustomization::ServiceProviderComboBox_OnGenerateWidget(TSharedPtr<FLocalizationServiceProviderWrapper> LSPWrapper) const
+//{
+//	ILocalizationServiceProvider* const LSP = LSPWrapper->Provider;
+//	return	SNew(STextBlock)
+//		.Text(LSP ? LSP->GetDisplayName() : LOCTEXT("NoServiceProviderName", "None"));
+//}
+//
+//void FLocalizationTargetSetDetailCustomization::ServiceProviderComboBox_OnSelectionChanged(TSharedPtr<FLocalizationServiceProviderWrapper> LSPWrapper, ESelectInfo::Type SelectInfo)
+//{
+//	ILocalizationServiceProvider* const LSP = LSPWrapper->Provider;
+//	FString ServiceProviderName = LSP ? LSP->GetName().ToString() : TEXT("None");
+//	ServiceProviderPropertyHandle->SetValue(ServiceProviderName);
+//
+//	if (LSP)
+//	{
+//		LSP->CustomizeSettingsDetails(*ServiceProviderCategoryBuilder);
+//	}
+//	DetailLayoutBuilder->ForceRefreshDetails();
+//}
+
+void FLocalizationTargetSetDetailCustomization::GatherAllTargets()
+{
+	for (ULocalizationTarget* const LocalizationTarget : TargetSet->TargetObjects)
+	{
+		// Save unsaved packages.
+		const bool bPromptUserToSave = true;
+		const bool bSaveMapPackages = true;
+		const bool bSaveContentPackages = true;
+		const bool bFastSave = false;
+		const bool bNotifyNoPackagesSaved = false;
+		const bool bCanBeDeclined = true;
+		bool DidPackagesNeedSaving;
+		const bool WerePackagesSaved = FEditorFileUtils::SaveDirtyPackages(bPromptUserToSave, bSaveMapPackages, bSaveContentPackages, bFastSave, bNotifyNoPackagesSaved, bCanBeDeclined, &DidPackagesNeedSaving);
+
+		if (DidPackagesNeedSaving && !WerePackagesSaved)
+		{
+			// Give warning dialog.
+			const FText MessageText = NSLOCTEXT("LocalizationCultureActions", "UnsavedPackagesWarningDialogMessage", "There are unsaved changes. These changes may not be gathered from correctly.");
+			const FText TitleText = NSLOCTEXT("LocalizationCultureActions", "UnsavedPackagesWarningDialogTitle", "Unsaved Changes Before Gather");
+			switch(FMessageDialog::Open(EAppMsgType::OkCancel, MessageText, &TitleText))
+			{
+			case EAppReturnType::Cancel:
+				{
+					return;
+				}
+				break;
+			}
+		}
+
+	}
+
+	// Execute gather.
+	const TSharedPtr<SWindow> ParentWindow = FSlateApplication::Get().FindWidgetWindow(DetailLayoutBuilder->GetDetailsView().AsShared());
+	TArray<FLocalizationTargetSettings*> TargetSettings;
+	for (ULocalizationTarget* const TargetObject : TargetSet->TargetObjects)
+	{
+		TargetSettings.Add(&TargetObject->Settings);
+	}
+	LocalizationCommandletTasks::GatherTargets(ParentWindow.ToSharedRef(), TargetSettings);
+
+	for (ULocalizationTarget* const LocalizationTarget : TargetSet->TargetObjects)
+	{
+		UpdateTargetFromReports(LocalizationTarget);
+	}
+}
+
+void FLocalizationTargetSetDetailCustomization::ImportAllTargets()
+{
+	IDesktopPlatform* DesktopPlatform = FDesktopPlatformModule::Get();
+	if (DesktopPlatform)
+	{
+		void* ParentWindowWindowHandle = NULL;
+		const TSharedPtr<SWindow> ParentWindow = FSlateApplication::Get().FindWidgetWindow(DetailLayoutBuilder->GetDetailsView().AsShared());
+		if (ParentWindow.IsValid() && ParentWindow->GetNativeWindow().IsValid())
+		{
+			ParentWindowWindowHandle = ParentWindow->GetNativeWindow()->GetOSWindowHandle();
+		}
+
+
+		const FString DefaultPath = FPaths::ConvertRelativePathToFull(FPaths::GameContentDir() / TEXT("Localization"));
+
+		// Prompt the user for the directory
+		FString OutputDirectory;
+		if (DesktopPlatform->OpenDirectoryDialog(ParentWindowWindowHandle, LOCTEXT("ImportAllTranslationsDialogTitle", "Import All Translations from Directory").ToString(), DefaultPath, OutputDirectory))
+		{
+			TArray<FLocalizationTargetSettings*> TargetSettings;
+			for (ULocalizationTarget* const TargetObject : TargetSet->TargetObjects)
+			{
+				TargetSettings.Add(&TargetObject->Settings);
+			}
+			LocalizationCommandletTasks::ImportTargets(ParentWindow.ToSharedRef(), TargetSettings, TOptional<FString>(OutputDirectory));
+
+			for (ULocalizationTarget* const LocalizationTarget : TargetSet->TargetObjects)
+			{
+				UpdateTargetFromReports(LocalizationTarget);
+			}
+		}
+	}
+}
+
+void FLocalizationTargetSetDetailCustomization::ExportAllTargets()
+{
+	IDesktopPlatform* DesktopPlatform = FDesktopPlatformModule::Get();
+	if (DesktopPlatform)
+	{
+		void* ParentWindowWindowHandle = NULL;
+		const TSharedPtr<SWindow> ParentWindow = FSlateApplication::Get().FindWidgetWindow(DetailLayoutBuilder->GetDetailsView().AsShared());
+		if (ParentWindow.IsValid() && ParentWindow->GetNativeWindow().IsValid())
+		{
+			ParentWindowWindowHandle = ParentWindow->GetNativeWindow()->GetOSWindowHandle();
+		}
+
+		const FString DefaultPath = FPaths::ConvertRelativePathToFull(FPaths::GameContentDir() / TEXT("Localization"));
+
+		// Prompt the user for the directory
+		FString OutputDirectory;
+		if (DesktopPlatform->OpenDirectoryDialog(ParentWindowWindowHandle, LOCTEXT("ExportAllTranslationsDialogTitle", "Export All Translations to Directory").ToString(), DefaultPath, OutputDirectory))
+		{
+			TArray<FLocalizationTargetSettings*> TargetSettings;
+			for (ULocalizationTarget* const TargetObject : TargetSet->TargetObjects)
+			{
+				TargetSettings.Add(&TargetObject->Settings);
+			}
+			LocalizationCommandletTasks::ExportTargets(ParentWindow.ToSharedRef(), TargetSettings, TOptional<FString>(OutputDirectory));
+		}
+	}
+}
+
+void FLocalizationTargetSetDetailCustomization::UpdateTargetFromReports(ULocalizationTarget* const LocalizationTarget)
+{
+	//TArray< TSharedPtr<IPropertyHandle> > WordCountPropertyHandles;
+
+	//const TSharedPtr<IPropertyHandle> TargetSettingsPropertyHandle = TargetEditor->GetTargetSettingsPropertyHandle();
+	//if (TargetSettingsPropertyHandle.IsValid() && TargetSettingsPropertyHandle->IsValidHandle())
+	//{
+	//	const TSharedPtr<IPropertyHandle> NativeCultureStatisticsPropertyHandle = TargetSettingsPropertyHandle->GetChildHandle(GET_MEMBER_NAME_CHECKED(FLocalizationTargetSettings, NativeCultureStatistics));
+	//	if (NativeCultureStatisticsPropertyHandle.IsValid() && NativeCultureStatisticsPropertyHandle->IsValidHandle())
+	//	{
+	//		const TSharedPtr<IPropertyHandle> WordCountPropertyHandle = NativeCultureStatisticsPropertyHandle->GetChildHandle(GET_MEMBER_NAME_CHECKED(FCultureStatistics, WordCount));
+	//		if (WordCountPropertyHandle.IsValid() && WordCountPropertyHandle->IsValidHandle())
+	//		{
+	//			WordCountPropertyHandles.Add(WordCountPropertyHandle);
+	//		}
+	//	}
+
+	//	const TSharedPtr<IPropertyHandle> SupportedCulturesStatisticsPropertyHandle = TargetSettingsPropertyHandle->GetChildHandle(GET_MEMBER_NAME_CHECKED(FLocalizationTargetSettings, SupportedCulturesStatistics));
+	//	if (SupportedCulturesStatisticsPropertyHandle.IsValid() && SupportedCulturesStatisticsPropertyHandle->IsValidHandle())
+	//	{
+	//		uint32 SupportedCultureCount = 0;
+	//		SupportedCulturesStatisticsPropertyHandle->GetNumChildren(SupportedCultureCount);
+	//		for (uint32 i = 0; i < SupportedCultureCount; ++i)
+	//		{
+	//			const TSharedPtr<IPropertyHandle> ElementPropertyHandle = SupportedCulturesStatisticsPropertyHandle->GetChildHandle(i);
+	//			if (ElementPropertyHandle.IsValid() && ElementPropertyHandle->IsValidHandle())
+	//			{
+	//				const TSharedPtr<IPropertyHandle> WordCountPropertyHandle = SupportedCulturesStatisticsPropertyHandle->GetChildHandle(GET_MEMBER_NAME_CHECKED(FCultureStatistics, WordCount));
+	//				if (WordCountPropertyHandle.IsValid() && WordCountPropertyHandle->IsValidHandle())
+	//				{
+	//					WordCountPropertyHandles.Add(WordCountPropertyHandle);
+	//				}
+	//			}
+	//		}
+	//	}
+	//}
+
+	//for (const TSharedPtr<IPropertyHandle>& WordCountPropertyHandle : WordCountPropertyHandles)
+	//{
+	//	WordCountPropertyHandle->NotifyPreChange();
+	//}
+	LocalizationTarget->Settings.UpdateWordCountsFromCSV();
+	LocalizationTarget->Settings.UpdateStatusFromConflictReport();
+	//for (const TSharedPtr<IPropertyHandle>& WordCountPropertyHandle : WordCountPropertyHandles)
+	//{
+	//	WordCountPropertyHandle->NotifyPostChange();
+	//}
+}
+
+void FLocalizationTargetSetDetailCustomization::CompileAllTargets()
+{
+	// Execute compile.
+	const TSharedPtr<SWindow> ParentWindow = FSlateApplication::Get().FindWidgetWindow(DetailLayoutBuilder->GetDetailsView().AsShared());
+	TArray<FLocalizationTargetSettings*> TargetSettings;
+	for (ULocalizationTarget* const TargetObject : TargetSet->TargetObjects)
+	{
+		TargetSettings.Add(&TargetObject->Settings);
+	}
+	LocalizationCommandletTasks::CompileTargets(ParentWindow.ToSharedRef(), TargetSettings);
+}
+
+TSharedRef<ITableRow> FLocalizationTargetSetDetailCustomization::OnGenerateRow(TSharedPtr<IPropertyHandle> TargetObjectPropertyHandle, const TSharedRef<STableViewBase>& Table)
+{
+	return SNew(SLocalizationDashboardTargetRow, Table, DetailLayoutBuilder->GetPropertyUtilities(), TargetObjectPropertyHandle.ToSharedRef());
+}
+
+FReply FLocalizationTargetSetDetailCustomization::OnNewTargetButtonClicked()
+{
+	if(TargetObjectsPropertyHandle.IsValid() && TargetObjectsPropertyHandle->IsValidHandle())
+	{
+		uint32 NewElementIndex;
+		TargetObjectsPropertyHandle->AsArray()->GetNumElements(NewElementIndex);
+		TargetObjectsPropertyHandle->AsArray()->AddItem();
+		NewEntryIndexToBeInitialized = NewElementIndex;
+	}
+
+	return FReply::Handled();
+}
+
+#undef LOCTEXT_NAMESPACE
\ No newline at end of file
diff --git a/Engine/Source/Editor/LocalizationDashboard/Private/LocalizationTargetSetDetailCustomization.h b/Engine/Source/Editor/LocalizationDashboard/Private/LocalizationTargetSetDetailCustomization.h
new file mode 100644
--- /dev/null
+++ b/Engine/Source/Editor/LocalizationDashboard/Private/LocalizationTargetSetDetailCustomization.h
@@ -0,0 +1,54 @@
+// Copyright 1998-2015 Epic Games, Inc. All Rights Reserved.
+
+#pragma once
+
+//class ILocalizationServiceProvider;
+
+//struct FLocalizationServiceProviderWrapper
+//{
+//	FLocalizationServiceProviderWrapper() : Provider(nullptr) {}
+//	FLocalizationServiceProviderWrapper(ILocalizationServiceProvider* const InProvider) : Provider(InProvider) {}
+//
+//	ILocalizationServiceProvider* Provider;
+//};
+
+class FLocalizationTargetSetDetailCustomization : public IDetailCustomization
+{
+public:
+	FLocalizationTargetSetDetailCustomization();
+	void CustomizeDetails(IDetailLayoutBuilder& DetailBuilder) override;
+
+private:
+	void BuildTargetsList();
+	void RebuildTargetsList();
+
+	//FText GetCurrentServiceProviderDisplayName() const;
+	//TSharedRef<SWidget> ServiceProviderComboBox_OnGenerateWidget(TSharedPtr<FLocalizationServiceProviderWrapper> LSPWrapper) const;
+	//void ServiceProviderComboBox_OnSelectionChanged(TSharedPtr<FLocalizationServiceProviderWrapper> LSPWrapper, ESelectInfo::Type SelectInfo);
+
+	void GatherAllTargets();
+	void ImportAllTargets();
+	void ExportAllTargets();
+	void UpdateTargetFromReports(ULocalizationTarget* const LocalizationTarget);
+	void CompileAllTargets();
+
+	TSharedRef<ITableRow> OnGenerateRow(TSharedPtr<IPropertyHandle> TargetObjectPropertyHandle, const TSharedRef<STableViewBase>& Table);
+	FReply OnNewTargetButtonClicked();
+
+private:
+	IDetailLayoutBuilder* DetailLayoutBuilder;
+
+	TWeakObjectPtr<ULocalizationTargetSet> TargetSet;
+
+	//TSharedPtr<IPropertyHandle> ServiceProviderPropertyHandle;
+	//IDetailCategoryBuilder* ServiceProviderCategoryBuilder;
+	//TArray< TSharedPtr<FLocalizationServiceProviderWrapper> > Providers;
+
+	TSharedPtr<IPropertyHandle> TargetObjectsPropertyHandle;
+	FSimpleDelegate TargetsArrayPropertyHandle_OnNumElementsChanged;
+	TArray< TSharedPtr<IPropertyHandle> > TargetsList;
+	TSharedPtr< SListView< TSharedPtr<IPropertyHandle> > > TargetsListView;
+
+	/* If set, the entry at the index specified needs to be initialized as soon as possible. */
+	int32 NewEntryIndexToBeInitialized;
+};
\ No newline at end of file
diff --git a/Engine/Source/Editor/LocalizationDashboard/Private/SLocalizationTargetEditor.cpp b/Engine/Source/Editor/LocalizationDashboard/Private/SLocalizationTargetEditor.cpp
--- a/Engine/Source/Editor/LocalizationDashboard/Private/SLocalizationTargetEditor.cpp
+++ b/Engine/Source/Editor/LocalizationDashboard/Private/SLocalizationTargetEditor.cpp
@@ -1,928 +1,6 @@
 // Copyright 1998-2014 Epic Games, Inc. All Rights Reserved.
 
 #include "LocalizationDashboardPrivatePCH.h"
-#include "IDetailCustomization.h"
-#include "EditorStyleSet.h"
-#include "DetailLayoutBuilder.h"
-#include "LocalizationConfigurationScript.h"
-#include "SCulturePicker.h"
-#include "SLocalizationTargetEditor.h"
-#include "PropertyEditorModule.h"
-#include "IDetailsView.h"
-#include "DetailWidgetRow.h"
-#include "DetailCategoryBuilder.h"
-#include "PropertyHandle.h"
-#include "SLocalizationTargetEditorCultureRow.h"
-#include "TargetsTableEntry.h"
-#include "SCulturePicker.h"
-#include "IDetailPropertyRow.h"
-#include "SLocalizationTargetStatusButton.h"
-#include "LocalizationCommandletTasks.h"
-
-#define LOCTEXT_NAMESPACE "LocalizationTargetEditor"
-
-namespace
-{
-	class FLocalizationTargetDetailCustomization : public IDetailCustomization
-	{
-	public:
-		FLocalizationTargetDetailCustomization(ULocalizationTargetSet* const InProjectSettings, ULocalizationTarget* const InLocalizationTarget);
-		virtual void CustomizeDetails(IDetailLayoutBuilder& DetailBuilder) override;
-
-		FLocalizationTargetSettings* GetTargetSettings() const;
-		TSharedPtr<IPropertyHandle> GetTargetSettingsPropertyHandle() const;
-
-	private:
-		FText GetTargetName() const;
-		bool IsTargetNameUnique(const FString& Name) const;
-		void OnTargetNameChanged(const FText& NewText);
-		void OnTargetNameCommitted(const FText& NewText, ETextCommit::Type Type);
-
-		void RebuildTargetDependenciesBox();
-		void RebuildTargetsList();
-		TSharedRef<ITableRow> OnGenerateTargetRow(ULocalizationTarget* OtherLocalizationTarget, const TSharedRef<STableViewBase>& Table);
-		void OnTargetDependencyCheckStateChanged(ULocalizationTarget* const OtherLocalizationTarget, const ECheckBoxState State);
-		ECheckBoxState IsTargetDependencyChecked(ULocalizationTarget* const OtherLocalizationTarget) const;
-
-		FText GetNativeCultureName() const;
-		FText GetNativeCultureDisplayName() const;
-		void OnNativeCultureSelected(FCulturePtr SelectedCulture, ESelectInfo::Type SelectInfo);
-
-		void Gather();
-		void ImportAllCultures();
-		void ExportAllCultures();
-		void Compile();
-		void RefreshWordCounts();
-		void UpdateTargetFromReports();
-
-		void BuildListedCulturesList();
-		void RebuildListedCulturesList();
-		TSharedRef<ITableRow> OnGenerateCultureRow(TSharedPtr<IPropertyHandle> CulturePropertyHandle, const TSharedRef<STableViewBase>& Table);
-
-		bool IsCultureSelectableAsSupported(FCulturePtr Culture);
-		void OnNewSupportedCultureSelected(FCulturePtr SelectedCulture, ESelectInfo::Type SelectInfo);
-
-	private:
-		ULocalizationTargetSet* ProjectSettings;
-		ULocalizationTarget* LocalizationTarget;
-
-		IDetailLayoutBuilder* DetailLayoutBuilder;
-		TSharedPtr<IPropertyHandle> TargetSettingsPropertyHandle;
-
-		TSharedPtr<SEditableTextBox> TargetNameEditableTextBox;
-
-		TSharedPtr<SHorizontalBox> TargetDependenciesHorizontalBox;
-		TArray< TSharedPtr<SWidget> > TargetDependenciesWidgets;
-		TArray<ULocalizationTarget*> TargetDependenciesOptionsList;
-		TSharedPtr< SListView<ULocalizationTarget*> > TargetDependenciesListView;
-
-		TSharedPtr<IPropertyHandle> NativeCultureStatisticsPropertyHandle;
-		TSharedPtr<IPropertyHandle> NativeCultureNamePropertyHandle;
-		TSharedPtr<SComboButton> NativeCultureComboButton;
-		TArray<FCulturePtr> AllCultures;
-
-		TSharedPtr<IPropertyHandle> SupportedCulturesStatisticsPropertyHandle;
-		FSimpleDelegate SupportedCulturesStatisticsPropertyHandle_OnNumElementsChanged;
-		TSharedPtr< SListView< TSharedPtr<IPropertyHandle> > > SupportedCultureListView;
-		TSharedPtr<SComboButton> AddNewSupportedCultureComboButton;
-		TSharedPtr<SCulturePicker> SupportedCulturePicker;
-		TArray< TSharedPtr<IPropertyHandle> > ListedCultureStatisticProperties;
-
-		/* If set, the entry at the index specified needs to be initialized as soon as possible. */
-		int32 NewEntryIndexToBeInitialized;
-		FCulturePtr SelectedNewCulture;
-	};
-
-	FLocalizationTargetDetailCustomization::FLocalizationTargetDetailCustomization(ULocalizationTargetSet* const InProjectSettings, ULocalizationTarget* const InLocalizationTarget)
-		: ProjectSettings(InProjectSettings)
-		, LocalizationTarget(InLocalizationTarget)
-		, DetailLayoutBuilder(nullptr)
-		, NewEntryIndexToBeInitialized(INDEX_NONE)
-	{
-	}
-
-	class FLocalizationTargetEditorCommands : public TCommands<FLocalizationTargetEditorCommands>
-	{
-	public:
-		FLocalizationTargetEditorCommands() 
-			: TCommands<FLocalizationTargetEditorCommands>("LocalizationTargetEditor", NSLOCTEXT("Contexts", "LocalizationTargetEditor", "Localization Target Editor"), NAME_None, FEditorStyle::GetStyleSetName())
-		{
-		}
-
-		TSharedPtr<FUICommandInfo> Gather;
-		TSharedPtr<FUICommandInfo> ImportAllCultures;
-		TSharedPtr<FUICommandInfo> ExportAllCultures;
-		TSharedPtr<FUICommandInfo> RefreshWordCounts;
-		TSharedPtr<FUICommandInfo> Compile;
-
-		/** Initialize commands */
-		virtual void RegisterCommands() override;
-	};
-
-	void FLocalizationTargetEditorCommands::RegisterCommands() 
-	{
-		UI_COMMAND( Gather, "Gather", "Gathers text for this target.", EUserInterfaceActionType::Button, FInputGesture() );
-		UI_COMMAND( ImportAllCultures, "Import All", "Imports translations for all cultures of this target.", EUserInterfaceActionType::Button, FInputGesture() );
-		UI_COMMAND( ExportAllCultures, "Export All", "Exports translations for all cultures of this target.", EUserInterfaceActionType::Button, FInputGesture() );
-		UI_COMMAND( RefreshWordCounts, "Count Words", "Recalculates the word counts for all cultures of this target.", EUserInterfaceActionType::Button, FInputGesture() );
-		UI_COMMAND( Compile, "Compile", "Compiles translations for this target.", EUserInterfaceActionType::Button, FInputGesture() );
-	}
-
-	void FLocalizationTargetDetailCustomization::CustomizeDetails(IDetailLayoutBuilder& DetailBuilder)
-	{
-		DetailLayoutBuilder = &DetailBuilder;
-		TargetSettingsPropertyHandle = DetailBuilder.GetProperty(GET_MEMBER_NAME_CHECKED(ULocalizationTarget, Settings));
-
-		if (TargetSettingsPropertyHandle.IsValid() && TargetSettingsPropertyHandle->IsValidHandle())
-		{
-			TargetSettingsPropertyHandle->MarkHiddenByCustomization();
-		}
-
-		//{
-		//	IDetailCategoryBuilder& ServiceProviderCategoryBuilder = DetailBuilder.EditCategory(FName("LocalizationServiceProvider"));
-
-		//	ILocalizationServiceProvider* const LSP = ILocalizationDashboardModule::Get().GetCurrentLocalizationServiceProvider();
-		//	if (LSP && LocalizationTarget)
-		//	{
-		//		LSP->CustomizeTargetDetails(ServiceProviderCategoryBuilder, *LocalizationTarget);
-		//	}
-		//}
-
-		if (TargetSettingsPropertyHandle.IsValid() && TargetSettingsPropertyHandle->IsValidHandle())
-		{
-			IDetailCategoryBuilder& DetailCategoryBuilder = DetailBuilder.EditCategory(FName("Target"));
-
-			const TSharedPtr<IPropertyHandle> NamePropertyHandle = TargetSettingsPropertyHandle->GetChildHandle(GET_MEMBER_NAME_CHECKED(FLocalizationTargetSettings, Name));
-			if (NamePropertyHandle.IsValid() && NamePropertyHandle->IsValidHandle())
-			{
-				FDetailWidgetRow& DetailWidgetRow = DetailCategoryBuilder.AddCustomRow( LOCTEXT("TargetNameFilterString", "Target Name") );
-				DetailWidgetRow.NameContent()
-					[
-						NamePropertyHandle->CreatePropertyNameWidget()
-					];
-				DetailWidgetRow.ValueContent()
-					[
-						SAssignNew(TargetNameEditableTextBox, SEditableTextBox)
-						.Text(this, &FLocalizationTargetDetailCustomization::GetTargetName)
-						.RevertTextOnEscape(true)
-						.OnTextChanged(this, &FLocalizationTargetDetailCustomization::OnTargetNameChanged)
-						.OnTextCommitted(this, &FLocalizationTargetDetailCustomization::OnTargetNameCommitted)
-					];
-			}
-			const TSharedPtr<IPropertyHandle> StatusPropertyHandle = TargetSettingsPropertyHandle->GetChildHandle(GET_MEMBER_NAME_CHECKED(FLocalizationTargetSettings, Status));
-			if (StatusPropertyHandle.IsValid() && StatusPropertyHandle->IsValidHandle())
-			{
-				StatusPropertyHandle->MarkHiddenByCustomization();
-				if (LocalizationTarget)
-				{
-					FDetailWidgetRow& StatusRow = DetailCategoryBuilder.AddCustomRow( LOCTEXT("ConflictReportStatusFilterString", "Conflict Report Status") );
-					StatusRow.NameContent()
-						[
-							StatusPropertyHandle->CreatePropertyNameWidget()
-						];
-					StatusRow.ValueContent()
-						.HAlign(HAlign_Left)
-						.VAlign(VAlign_Center)
-						[
-							SNew(SLocalizationTargetStatusButton, *LocalizationTarget)
-						];
-				}
-			}
-
-			const TSharedPtr<IPropertyHandle> TargetDependenciesPropertyHandle = TargetSettingsPropertyHandle->GetChildHandle(GET_MEMBER_NAME_CHECKED(FLocalizationTargetSettings, TargetDependencies));
-			if (TargetDependenciesPropertyHandle.IsValid() && TargetDependenciesPropertyHandle->IsValidHandle())
-			{
-				TargetDependenciesPropertyHandle->MarkHiddenByCustomization();
-				if (LocalizationTarget)
-				{
-					FDetailWidgetRow& TargetDependenciesRow = DetailCategoryBuilder.AddCustomRow( LOCTEXT("ConflictReportStatusFilterString", "Conflict Report Status") );
-					TargetDependenciesRow.NameContent()
-						[
-							TargetDependenciesPropertyHandle->CreatePropertyNameWidget()
-						];
-					TargetDependenciesRow.ValueContent()
-						.HAlign(HAlign_Left)
-						.VAlign(VAlign_Center)
-						[
-							SNew(SComboButton)
-							.ButtonContent()
-							[
-								SAssignNew(TargetDependenciesHorizontalBox, SHorizontalBox)
-							]
-							.HasDownArrow(true)
-								.OnGetMenuContent_Lambda([this]() -> TSharedRef<SWidget>
-							{
-								RebuildTargetsList();
-
-								if (TargetDependenciesOptionsList.Num() > 0)
-								{
-									return SNew(SBox)
-										.MaxDesiredHeight(400.0f)
-										.MaxDesiredWidth(300.0f)
-										[
-											SNew(SListView<ULocalizationTarget*>)
-											.SelectionMode(ESelectionMode::None)
-											.ListItemsSource(&TargetDependenciesOptionsList)
-											.OnGenerateRow(this, &FLocalizationTargetDetailCustomization::OnGenerateTargetRow)
-										];
-								}
-								else
-								{
-									return SNullWidget::NullWidget;
-								}
-							})
-						];
-
-					RebuildTargetDependenciesBox();
-				}
-			}
-
-			const TSharedPtr<IPropertyHandle> AdditionalManifestDependenciesPropertyHandle = TargetSettingsPropertyHandle->GetChildHandle(GET_MEMBER_NAME_CHECKED(FLocalizationTargetSettings, AdditionalManifestDependencies));
-			if (AdditionalManifestDependenciesPropertyHandle.IsValid() && AdditionalManifestDependenciesPropertyHandle->IsValidHandle())
-			{
-				DetailCategoryBuilder.AddProperty(AdditionalManifestDependenciesPropertyHandle);
-			}
-		}
-
-		if (TargetSettingsPropertyHandle.IsValid() && TargetSettingsPropertyHandle->IsValidHandle())
-		{
-			IDetailCategoryBuilder& DetailCategoryBuilder = DetailBuilder.EditCategory(FName("Cultures"));
-
-			NativeCultureStatisticsPropertyHandle = TargetSettingsPropertyHandle->GetChildHandle(GET_MEMBER_NAME_CHECKED(FLocalizationTargetSettings, NativeCultureStatistics));
-			if (NativeCultureStatisticsPropertyHandle.IsValid() && NativeCultureStatisticsPropertyHandle->IsValidHandle())
-			{
-				NativeCultureNamePropertyHandle = NativeCultureStatisticsPropertyHandle->GetChildHandle(GET_MEMBER_NAME_CHECKED(FCultureStatistics, CultureName));
-				if (NativeCultureNamePropertyHandle.IsValid() && NativeCultureNamePropertyHandle->IsValidHandle())
-				{
-					NativeCultureNamePropertyHandle->MarkHiddenByCustomization();
-					FDetailWidgetRow& NativeCultureRow = DetailCategoryBuilder.AddCustomRow( LOCTEXT("NativeCultureNameFilterString", "Native Culture Name") );
-
-					FString NativeCultureName;
-					NativeCultureNamePropertyHandle->GetValue(NativeCultureName);
-					FCulturePtr NativeCulture = FInternationalization::Get().GetCulture(NativeCultureName);
-
-					NativeCultureRow.NameContent()
-						[
-							NativeCultureNamePropertyHandle->CreatePropertyNameWidget(LOCTEXT("NativeCultureNameLabel", "Native Culture"))
-						];
-
-					NativeCultureRow.ValueContent()
-						[
-							SAssignNew(NativeCultureComboButton, SComboButton)
-							.ButtonContent()
-							[
-								SNew(STextBlock)
-								.Text(this, &FLocalizationTargetDetailCustomization::GetNativeCultureDisplayName)
-								.ToolTipText(this, &FLocalizationTargetDetailCustomization::GetNativeCultureName)
-							]
-							.HasDownArrow(true)
-								.MenuContent()
-								[
-									SNew(SBox)
-									.MaxDesiredHeight(400.0f)
-									.MaxDesiredWidth(300.0f)
-									[
-										SNew(SCulturePicker)
-										.OnSelectionChanged(this, &FLocalizationTargetDetailCustomization::OnNativeCultureSelected)
-										.InitialSelection(NativeCulture)
-									]
-								]
-						];
-				}
-			}
-
-			SupportedCulturesStatisticsPropertyHandle = TargetSettingsPropertyHandle->GetChildHandle(GET_MEMBER_NAME_CHECKED(FLocalizationTargetSettings, SupportedCulturesStatistics));
-			if (SupportedCulturesStatisticsPropertyHandle.IsValid() && SupportedCulturesStatisticsPropertyHandle->IsValidHandle())
-			{
-				SupportedCulturesStatisticsPropertyHandle->MarkHiddenByCustomization();
-				SupportedCulturesStatisticsPropertyHandle_OnNumElementsChanged = FSimpleDelegate::CreateSP(this, &FLocalizationTargetDetailCustomization::RebuildListedCulturesList);
-				SupportedCulturesStatisticsPropertyHandle->AsArray()->SetOnNumElementsChanged(SupportedCulturesStatisticsPropertyHandle_OnNumElementsChanged);
-
-				FLocalizationTargetEditorCommands::Register();
-				const TSharedRef< FUICommandList > CommandList = MakeShareable(new FUICommandList);
-				FToolBarBuilder ToolBarBuilder( CommandList, FMultiBoxCustomization::AllowCustomization("LocalizationTargetEditor") );
-
-				CommandList->MapAction( FLocalizationTargetEditorCommands::Get().Gather, FExecuteAction::CreateSP(this, &FLocalizationTargetDetailCustomization::Gather));
-				ToolBarBuilder.AddToolBarButton(FLocalizationTargetEditorCommands::Get().Gather, NAME_None, TAttribute<FText>(), TAttribute<FText>(), FSlateIcon(FEditorStyle::GetStyleSetName(), "LocalizationTargetEditor.Gather"));
-
-				CommandList->MapAction( FLocalizationTargetEditorCommands::Get().ImportAllCultures, FExecuteAction::CreateSP(this, &FLocalizationTargetDetailCustomization::ImportAllCultures));
-				ToolBarBuilder.AddToolBarButton(FLocalizationTargetEditorCommands::Get().ImportAllCultures, NAME_None, TAttribute<FText>(), TAttribute<FText>(), FSlateIcon(FEditorStyle::GetStyleSetName(), "LocalizationTargetEditor.ImportForAllCultures"));
-
-				CommandList->MapAction( FLocalizationTargetEditorCommands::Get().ExportAllCultures, FExecuteAction::CreateSP(this, &FLocalizationTargetDetailCustomization::ExportAllCultures));
-				ToolBarBuilder.AddToolBarButton(FLocalizationTargetEditorCommands::Get().ExportAllCultures, NAME_None, TAttribute<FText>(), TAttribute<FText>(), FSlateIcon(FEditorStyle::GetStyleSetName(), "LocalizationTargetEditor.ExportForAllCultures"));
-
-				CommandList->MapAction( FLocalizationTargetEditorCommands::Get().RefreshWordCounts, FExecuteAction::CreateSP(this, &FLocalizationTargetDetailCustomization::RefreshWordCounts));
-				ToolBarBuilder.AddToolBarButton(FLocalizationTargetEditorCommands::Get().RefreshWordCounts, NAME_None, TAttribute<FText>(), TAttribute<FText>(), FSlateIcon(FEditorStyle::GetStyleSetName(), "LocalizationTargetEditor.RefreshWordCounts"));
-
-				CommandList->MapAction( FLocalizationTargetEditorCommands::Get().Compile, FExecuteAction::CreateSP(this, &FLocalizationTargetDetailCustomization::Compile));
-				ToolBarBuilder.AddToolBarButton(FLocalizationTargetEditorCommands::Get().Compile, NAME_None, TAttribute<FText>(), TAttribute<FText>(), FSlateIcon(FEditorStyle::GetStyleSetName(), "LocalizationTargetEditor.Compile"));
-
-
-				BuildListedCulturesList();
-
-				DetailCategoryBuilder.AddCustomRow( LOCTEXT("SupportedCultureNamesFilterString", "Supported Culture Names") )
-					.WholeRowContent()
-					[
-						SNew(SVerticalBox)
-						+SVerticalBox::Slot()
-						.AutoHeight()
-						[
-							ToolBarBuilder.MakeWidget()
-						]
-						+SVerticalBox::Slot()
-							.AutoHeight()
-							[
-								SAssignNew(SupportedCultureListView, SListView< TSharedPtr<IPropertyHandle> >)
-								.OnGenerateRow(this, &FLocalizationTargetDetailCustomization::OnGenerateCultureRow)
-								.ListItemsSource(&ListedCultureStatisticProperties)
-								.SelectionMode(ESelectionMode::None)
-								.HeaderRow
-								(
-								SNew(SHeaderRow)
-								+SHeaderRow::Column("Culture")
-								.DefaultLabel( NSLOCTEXT("LocalizationCulture", "CultureColumnLabel", "Culture"))
-								.HAlignHeader(HAlign_Fill)
-								.HAlignCell(HAlign_Fill)
-								.VAlignCell(VAlign_Center)
-								+SHeaderRow::Column("WordCount")
-								.DefaultLabel( NSLOCTEXT("LocalizationCulture", "WordCountColumnLabel", "Word Count"))
-								.HAlignHeader(HAlign_Center)
-								.HAlignCell(HAlign_Fill)
-								.VAlignCell(VAlign_Center)
-								+SHeaderRow::Column("Actions")
-								.DefaultLabel( NSLOCTEXT("LocalizationCulture", "ActionsColumnLabel", "Actions"))
-								.HAlignHeader(HAlign_Center)
-								.HAlignCell(HAlign_Center)
-								.VAlignCell(VAlign_Center)
-								)
-							]
-						+SVerticalBox::Slot()
-							.AutoHeight()
-							.VAlign(VAlign_Center)
-							[
-								SAssignNew(AddNewSupportedCultureComboButton, SComboButton)
-								.ButtonContent()
-								[
-									SNew(STextBlock)
-									.Text(NSLOCTEXT("LocalizationCulture", "AddNewCultureButtonLabel", "Add New Culture"))
-								]
-								.MenuContent()
-									[
-										SNew(SBox)
-										.MaxDesiredHeight(400.0f)
-										.MaxDesiredWidth(300.0f)
-										[
-											SAssignNew(SupportedCulturePicker, SCulturePicker)
-											.OnSelectionChanged(this, &FLocalizationTargetDetailCustomization::OnNewSupportedCultureSelected)
-											.IsCulturePickable(this, &FLocalizationTargetDetailCustomization::IsCultureSelectableAsSupported)
-										]
-									]
-							]
-					];
-			}
-		}
-	}
-
-	FLocalizationTargetSettings* FLocalizationTargetDetailCustomization::GetTargetSettings() const
-	{
-		return LocalizationTarget ? &(LocalizationTarget->Settings) : nullptr;
-	}
-
-	TSharedPtr<IPropertyHandle> FLocalizationTargetDetailCustomization::GetTargetSettingsPropertyHandle() const
-	{
-		return TargetSettingsPropertyHandle;
-	}
-
-	FText FLocalizationTargetDetailCustomization::GetTargetName() const
-	{
-		if (LocalizationTarget)
-		{
-			return FText::FromString(LocalizationTarget->Settings.Name);
-		}
-		return FText::GetEmpty();
-	}
-
-	bool FLocalizationTargetDetailCustomization::IsTargetNameUnique(const FString& Name) const
-	{
-		for (ULocalizationTarget* const TargetObject : ProjectSettings->TargetObjects)
-		{
-			if (TargetObject != LocalizationTarget)
-			{
-				if (TargetObject->Settings.Name == LocalizationTarget->Settings.Name)
-				{
-					return false;
-				}
-			}
-		}
-		return true;
-	}
-
-	void FLocalizationTargetDetailCustomization::OnTargetNameChanged(const FText& NewText)
-	{
-		const FString& NewName = NewText.ToString();
-
-		// TODO: Target names must be valid directory names, because they are used as directory names.
-		// ValidatePath allows /, which is not a valid directory name character
-		FText Error;
-		if (!FPaths::ValidatePath(NewName, &Error))
-		{
-			TargetNameEditableTextBox->SetError(Error);
-			return;		
-		}
-
-		// Target name must be unique.
-		if (!IsTargetNameUnique(NewName))
-		{
-			TargetNameEditableTextBox->SetError(LOCTEXT("DuplicateTargetNameError", "Target name must be unique."));
-			return;
-		}
-
-		// Clear error if nothing has failed.
-		TargetNameEditableTextBox->SetError(FText::GetEmpty());
-	}
-
-	void FLocalizationTargetDetailCustomization::OnTargetNameCommitted(const FText& NewText, ETextCommit::Type Type)
-	{
-		// Target name must be unique.
-		if (!IsTargetNameUnique(NewText.ToString()))
-		{
-			return;
-		}
-
-		if (TargetSettingsPropertyHandle.IsValid() && TargetSettingsPropertyHandle->IsValidHandle())
-		{
-			FLocalizationTargetSettings* const TargetSettings = GetTargetSettings();
-			if (TargetSettings)
-			{
-				// Early out if the committed name is the same as the current name.
-				if (TargetSettings->Name == NewText.ToString())
-				{
-					return;
-				}
-
-				const TSharedPtr<IPropertyHandle> TargetNamePropertyHandle = TargetSettingsPropertyHandle->GetChildHandle(GET_MEMBER_NAME_CHECKED(FLocalizationTargetSettings, Name));
-
-				if (TargetNamePropertyHandle.IsValid() && TargetNamePropertyHandle->IsValidHandle())
-				{
-					TargetNamePropertyHandle->NotifyPreChange();
-				}
-
-				TargetSettings->Rename(NewText.ToString());
-
-				if (TargetNamePropertyHandle.IsValid() && TargetNamePropertyHandle->IsValidHandle())
-				{
-					TargetNamePropertyHandle->NotifyPostChange();
-				}
-			}
-		}
-	}
-
-	void FLocalizationTargetDetailCustomization::RebuildTargetDependenciesBox()
-	{
-		if (TargetDependenciesHorizontalBox.IsValid())
-		{
-			for (const TSharedPtr<SWidget>& Widget : TargetDependenciesWidgets)
-			{
-				TargetDependenciesHorizontalBox->RemoveSlot(Widget.ToSharedRef());
-			}
-			TargetDependenciesWidgets.Empty();
-
-			for (const FString& TargetDependency : LocalizationTarget->Settings.TargetDependencies)
-			{
-				const TSharedRef<SWidget> Widget = SNew(SBorder)
-					[
-						SNew(STextBlock)
-						.Text(FText::FromString(TargetDependency))
-					];
-
-				TargetDependenciesWidgets.Add(Widget);
-				TargetDependenciesHorizontalBox->AddSlot()
-					[
-						Widget
-					];
-			}
-		}
-	}
-
-	void FLocalizationTargetDetailCustomization::RebuildTargetsList()
-	{
-		TargetDependenciesOptionsList.Empty();
-		TFunction<bool (ULocalizationTarget* const)> DoesTargetDependOnUs;
-		DoesTargetDependOnUs = [&, this](ULocalizationTarget* const OtherTarget) -> bool
-		{
-			if (OtherTarget->Settings.TargetDependencies.Contains(LocalizationTarget->Settings.Name))
-			{
-				return true;
-			}
-
-			for (const FString& OtherTargetDependencyName : OtherTarget->Settings.TargetDependencies)
-			{
-				ULocalizationTarget** const OtherTargetDependency = ProjectSettings->TargetObjects.FindByPredicate([OtherTargetDependencyName](ULocalizationTarget* SomeLocalizationTarget)->bool{return SomeLocalizationTarget->Settings.Name == OtherTargetDependencyName;});
-				if (OtherTargetDependency && DoesTargetDependOnUs(*OtherTargetDependency))
-				{
-					return true;
-				}
-			}
-
-			return false;
-		};
-
-		for (ULocalizationTarget* const OtherTarget : ProjectSettings->TargetObjects)
-		{
-			if (OtherTarget != LocalizationTarget)
-			{
-				if (!DoesTargetDependOnUs(OtherTarget))
-				{
-					TargetDependenciesOptionsList.Add(OtherTarget);
-				}
-			}
-		}
-
-		if (TargetDependenciesListView.IsValid())
-		{
-			TargetDependenciesListView->RequestListRefresh();
-		}
-	}
-
-	TSharedRef<ITableRow> FLocalizationTargetDetailCustomization::OnGenerateTargetRow(ULocalizationTarget* OtherLocalizationTarget, const TSharedRef<STableViewBase>& Table)
-	{
-		return SNew(STableRow<ULocalizationTarget*>, Table)
-			.ShowSelection(true)
-			.Content()
-			[
-				SNew(SCheckBox)
-				.OnCheckStateChanged_Lambda([this, OtherLocalizationTarget](ECheckBoxState State)
-				{
-					OnTargetDependencyCheckStateChanged(OtherLocalizationTarget, State);
-				})
-					.IsChecked_Lambda([this, OtherLocalizationTarget]()->ECheckBoxState
-				{
-					return IsTargetDependencyChecked(OtherLocalizationTarget);
-				})
-					.Content()
-					[
-						SNew(STextBlock)
-						.Text(FText::FromString(OtherLocalizationTarget->Settings.Name))
-					]
-			];
-	}
-
-	void FLocalizationTargetDetailCustomization::OnTargetDependencyCheckStateChanged(ULocalizationTarget* const OtherLocalizationTarget, const ECheckBoxState State)
-	{
-		const TSharedPtr<IPropertyHandle> TargetDependenciesPropertyHandle = TargetSettingsPropertyHandle->GetChildHandle(GET_MEMBER_NAME_CHECKED(FLocalizationTargetSettings, TargetDependencies));
-
-		if (TargetDependenciesPropertyHandle.IsValid() && TargetDependenciesPropertyHandle->IsValidHandle())
-		{
-			TargetDependenciesPropertyHandle->NotifyPreChange();
-		}
-
-		switch (State)
-		{
-		case ECheckBoxState::Checked:
-			LocalizationTarget->Settings.TargetDependencies.Add(OtherLocalizationTarget->Settings.Name);
-			break;
-		case ECheckBoxState::Unchecked:
-			LocalizationTarget->Settings.TargetDependencies.Remove(OtherLocalizationTarget->Settings.Name);
-			break;
-		}
-
-		if (TargetDependenciesPropertyHandle.IsValid() && TargetDependenciesPropertyHandle->IsValidHandle())
-		{
-			TargetDependenciesPropertyHandle->NotifyPostChange();
-		}
-
-		RebuildTargetDependenciesBox();
-	}
-
-	ECheckBoxState FLocalizationTargetDetailCustomization::IsTargetDependencyChecked(ULocalizationTarget* const OtherLocalizationTarget) const
-	{
-		return LocalizationTarget->Settings.TargetDependencies.Contains(OtherLocalizationTarget->Settings.Name) ? ECheckBoxState::Checked : ECheckBoxState::Unchecked;
-	}
-
-	FText FLocalizationTargetDetailCustomization::GetNativeCultureName() const
-	{
-		FCulturePtr NativeCulture;
-
-		if(NativeCultureNamePropertyHandle.IsValid() && NativeCultureNamePropertyHandle->IsValidHandle())
-		{
-			FString Value;
-			NativeCultureNamePropertyHandle->GetValue(Value);
-			NativeCulture = FInternationalization::Get().GetCulture(Value);
-		}
-
-		return NativeCulture.IsValid() ? FText::FromString(NativeCulture->GetName()) : FText::GetEmpty();
-	}
-
-	FText FLocalizationTargetDetailCustomization::GetNativeCultureDisplayName() const
-	{
-		FCulturePtr NativeCulture;
-
-		if(NativeCultureNamePropertyHandle.IsValid() && NativeCultureNamePropertyHandle->IsValidHandle())
-		{
-			FString Value;
-			NativeCultureNamePropertyHandle->GetValue(Value);
-			NativeCulture = FInternationalization::Get().GetCulture(Value);
-		}
-
-		return NativeCulture.IsValid() ? FText::FromString(NativeCulture->GetDisplayName()) : FText::GetEmpty();
-	}
-
-	void FLocalizationTargetDetailCustomization::OnNativeCultureSelected(FCulturePtr SelectedCulture, ESelectInfo::Type SelectInfo)
-	{
-		if(NativeCultureNamePropertyHandle.IsValid() && NativeCultureNamePropertyHandle->IsValidHandle())
-		{
-			NativeCultureNamePropertyHandle->SetValue(SelectedCulture->GetName());
-		}
-		if (NativeCultureComboButton.IsValid())
-		{
-			NativeCultureComboButton->SetIsOpen(false);
-		}
-
-		// TODO:	If the selected culture is an existing supported culture:
-		//			copy that supported culture's statistics to the native culture statistics
-		//			remove the supported culture's statistics
-
-		// TODO: Warn about this action invalidating all word counts and, realistically, invalidating all existing translations.
-		// Possibly offer an option to delete all archives?
-	}
-
-	void FLocalizationTargetDetailCustomization::Gather()
-	{
-		if (LocalizationTarget)
-		{
-			// Save unsaved packages.
-			const bool bPromptUserToSave = true;
-			const bool bSaveMapPackages = true;
-			const bool bSaveContentPackages = true;
-			const bool bFastSave = false;
-			const bool bNotifyNoPackagesSaved = false;
-			const bool bCanBeDeclined = true;
-			bool DidPackagesNeedSaving;
-			const bool WerePackagesSaved = FEditorFileUtils::SaveDirtyPackages(bPromptUserToSave, bSaveMapPackages, bSaveContentPackages, bFastSave, bNotifyNoPackagesSaved, bCanBeDeclined, &DidPackagesNeedSaving);
-
-			if (DidPackagesNeedSaving && !WerePackagesSaved)
-			{
-				// Give warning dialog.
-				const FText MessageText = NSLOCTEXT("LocalizationCultureActions", "UnsavedPackagesWarningDialogMessage", "There are unsaved changes. These changes may not be gathered from correctly.");
-				const FText TitleText = NSLOCTEXT("LocalizationCultureActions", "UnsavedPackagesWarningDialogTitle", "Unsaved Changes Before Gather");
-				switch(FMessageDialog::Open(EAppMsgType::OkCancel, MessageText, &TitleText))
-				{
-				case EAppReturnType::Cancel:
-					{
-						return;
-					}
-					break;
-				}
-			}
-
-			// Execute gather.
-			const TSharedPtr<SWindow> ParentWindow = FSlateApplication::Get().FindWidgetWindow(DetailLayoutBuilder->GetDetailsView().AsShared());
-			LocalizationCommandletTasks::GatherTarget(ParentWindow.ToSharedRef(), LocalizationTarget->Settings);
-
-			UpdateTargetFromReports();
-		}
-	}
-
-	void FLocalizationTargetDetailCustomization::ImportAllCultures()
-	{
-		IDesktopPlatform* DesktopPlatform = FDesktopPlatformModule::Get();
-		if (LocalizationTarget && DesktopPlatform)
-		{
-			void* ParentWindowWindowHandle = NULL;
-			const TSharedPtr<SWindow> ParentWindow = FSlateApplication::Get().FindWidgetWindow(DetailLayoutBuilder->GetDetailsView().AsShared());
-			if (ParentWindow.IsValid() && ParentWindow->GetNativeWindow().IsValid())
-			{
-				ParentWindowWindowHandle = ParentWindow->GetNativeWindow()->GetOSWindowHandle();
-			}
-
-			const FString DefaultPath = FPaths::ConvertRelativePathToFull(LocalizationConfigurationScript::GetDataDirectory(LocalizationTarget->Settings));
-
-			FText DialogTitle;
-			{
-				FFormatNamedArguments FormatArguments;
-				FormatArguments.Add(TEXT("TargetName"), FText::FromString(LocalizationTarget->Settings.Name));
-				DialogTitle = FText::Format(LOCTEXT("ImportAllTranslationsForTargetDialogTitleFormat", "Import All Translations for {TargetName} from Directory"), FormatArguments);
-			}
-
-			// Prompt the user for the directory
-			FString OutputDirectory;
-			if (DesktopPlatform->OpenDirectoryDialog(ParentWindowWindowHandle, DialogTitle.ToString(), DefaultPath, OutputDirectory))
-			{
-				LocalizationCommandletTasks::ImportTarget(ParentWindow.ToSharedRef(), LocalizationTarget->Settings, TOptional<FString>(OutputDirectory));
-
-				UpdateTargetFromReports();
-			}
-		}
-	}
-
-	void FLocalizationTargetDetailCustomization::ExportAllCultures()
-	{
-		IDesktopPlatform* DesktopPlatform = FDesktopPlatformModule::Get();
-		if (LocalizationTarget && DesktopPlatform)
-		{
-			void* ParentWindowWindowHandle = NULL;
-			const TSharedPtr<SWindow> ParentWindow = FSlateApplication::Get().FindWidgetWindow(DetailLayoutBuilder->GetDetailsView().AsShared());
-			if (ParentWindow.IsValid() && ParentWindow->GetNativeWindow().IsValid())
-			{
-				ParentWindowWindowHandle = ParentWindow->GetNativeWindow()->GetOSWindowHandle();
-			}
-
-			const FString DefaultPath = FPaths::ConvertRelativePathToFull(LocalizationConfigurationScript::GetDataDirectory(LocalizationTarget->Settings));
-
-			FText DialogTitle;
-			{
-				FFormatNamedArguments FormatArguments;
-				FormatArguments.Add(TEXT("TargetName"), FText::FromString(LocalizationTarget->Settings.Name));
-				DialogTitle = FText::Format(LOCTEXT("ExportAllTranslationsForTargetDialogTitleFormat", "Export All Translations for {TargetName} to Directory"), FormatArguments);
-			}
-
-			// Prompt the user for the directory
-			FString OutputDirectory;
-			if (DesktopPlatform->OpenDirectoryDialog(ParentWindowWindowHandle, DialogTitle.ToString(), DefaultPath, OutputDirectory))
-			{
-				LocalizationCommandletTasks::ExportTarget(ParentWindow.ToSharedRef(), LocalizationTarget->Settings, TOptional<FString>(OutputDirectory));
-			}
-		}
-	}
-
-	void FLocalizationTargetDetailCustomization::Compile()
-	{
-		if (LocalizationTarget)
-		{
-			// Execute compile.
-			const TSharedPtr<SWindow> ParentWindow = FSlateApplication::Get().FindWidgetWindow(DetailLayoutBuilder->GetDetailsView().AsShared());
-			LocalizationCommandletTasks::CompileTarget(ParentWindow.ToSharedRef(), LocalizationTarget->Settings);
-		}
-	}
-
-	void FLocalizationTargetDetailCustomization::RefreshWordCounts()
-	{
-		if (LocalizationTarget)
-		{
-			const TSharedPtr<SWindow> ParentWindow = FSlateApplication::Get().FindWidgetWindow(DetailLayoutBuilder->GetDetailsView().AsShared());
-			LocalizationCommandletTasks::GenerateReportsForTarget(ParentWindow.ToSharedRef(), LocalizationTarget->Settings);
-
-			UpdateTargetFromReports();
-		}
-	}
-
-	void FLocalizationTargetDetailCustomization::UpdateTargetFromReports()
-	{
-		FLocalizationTargetSettings* const TargetSettings = GetTargetSettings();
-		if (TargetSettings)
-		{
-			TArray< TSharedPtr<IPropertyHandle> > WordCountPropertyHandles;
-
-			if (TargetSettingsPropertyHandle.IsValid() && TargetSettingsPropertyHandle->IsValidHandle())
-			{
-				const TSharedPtr<IPropertyHandle> NativeCultureStatisticsPropertyHandle = TargetSettingsPropertyHandle->GetChildHandle(GET_MEMBER_NAME_CHECKED(FLocalizationTargetSettings, NativeCultureStatistics));
-				if (NativeCultureStatisticsPropertyHandle.IsValid() && NativeCultureStatisticsPropertyHandle->IsValidHandle())
-				{
-					const TSharedPtr<IPropertyHandle> WordCountPropertyHandle = NativeCultureStatisticsPropertyHandle->GetChildHandle(GET_MEMBER_NAME_CHECKED(FCultureStatistics, WordCount));
-					if (WordCountPropertyHandle.IsValid() && WordCountPropertyHandle->IsValidHandle())
-					{
-						WordCountPropertyHandles.Add(WordCountPropertyHandle);
-					}
-				}
-
-				const TSharedPtr<IPropertyHandle> SupportedCulturesStatisticsPropertyHandle = TargetSettingsPropertyHandle->GetChildHandle(GET_MEMBER_NAME_CHECKED(FLocalizationTargetSettings, SupportedCulturesStatistics));
-				if (SupportedCulturesStatisticsPropertyHandle.IsValid() && SupportedCulturesStatisticsPropertyHandle->IsValidHandle())
-				{
-					uint32 SupportedCultureCount = 0;
-					SupportedCulturesStatisticsPropertyHandle->GetNumChildren(SupportedCultureCount);
-					for (uint32 i = 0; i < SupportedCultureCount; ++i)
-					{
-						const TSharedPtr<IPropertyHandle> ElementPropertyHandle = SupportedCulturesStatisticsPropertyHandle->GetChildHandle(i);
-						if (ElementPropertyHandle.IsValid() && ElementPropertyHandle->IsValidHandle())
-						{
-							const TSharedPtr<IPropertyHandle> WordCountPropertyHandle = SupportedCulturesStatisticsPropertyHandle->GetChildHandle(GET_MEMBER_NAME_CHECKED(FCultureStatistics, WordCount));
-							if (WordCountPropertyHandle.IsValid() && WordCountPropertyHandle->IsValidHandle())
-							{
-								WordCountPropertyHandles.Add(WordCountPropertyHandle);
-							}
-						}
-					}
-				}
-			}
-
-			for (const TSharedPtr<IPropertyHandle>& WordCountPropertyHandle : WordCountPropertyHandles)
-			{
-				WordCountPropertyHandle->NotifyPreChange();
-			}
-			TargetSettings->UpdateWordCountsFromCSV();
-			TargetSettings->UpdateStatusFromConflictReport();
-			for (const TSharedPtr<IPropertyHandle>& WordCountPropertyHandle : WordCountPropertyHandles)
-			{
-				WordCountPropertyHandle->NotifyPostChange();
-			}
-		}
-	}
-
-	void FLocalizationTargetDetailCustomization::BuildListedCulturesList()
-	{
-		ListedCultureStatisticProperties.AddUnique(NativeCultureStatisticsPropertyHandle);
-		const TSharedPtr<IPropertyHandleArray> SupportedCulturesStatisticsArrayPropertyHandle = SupportedCulturesStatisticsPropertyHandle->AsArray();
-		if (SupportedCulturesStatisticsArrayPropertyHandle.IsValid())
-		{
-			uint32 ElementCount = 0;
-			SupportedCulturesStatisticsArrayPropertyHandle->GetNumElements(ElementCount);
-			for (uint32 i = 0; i < ElementCount; ++i)
-			{
-				const TSharedPtr<IPropertyHandle> CultureStatisticsProperty = SupportedCulturesStatisticsArrayPropertyHandle->GetElement(i);
-				ListedCultureStatisticProperties.AddUnique(CultureStatisticsProperty);
-			}
-		}
-	}
-
-	void FLocalizationTargetDetailCustomization::RebuildListedCulturesList()
-	{
-		if (NewEntryIndexToBeInitialized != INDEX_NONE)
-		{
-			const TSharedPtr<IPropertyHandle> SupportedCultureStatisticsPropertyHandle = SupportedCulturesStatisticsPropertyHandle->GetChildHandle(NewEntryIndexToBeInitialized);
-
-			const TSharedPtr<IPropertyHandle> CultureNamePropertyHandle = SupportedCultureStatisticsPropertyHandle->GetChildHandle(GET_MEMBER_NAME_CHECKED(FCultureStatistics, CultureName));
-			if(CultureNamePropertyHandle.IsValid() && CultureNamePropertyHandle->IsValidHandle())
-			{
-				CultureNamePropertyHandle->SetValue(SelectedNewCulture->GetName());
-			}
-
-			const TSharedPtr<IPropertyHandle> WordCountPropertyHandle = SupportedCultureStatisticsPropertyHandle->GetChildHandle(GET_MEMBER_NAME_CHECKED(FCultureStatistics, WordCount));
-			if(WordCountPropertyHandle.IsValid() && WordCountPropertyHandle->IsValidHandle())
-			{
-				WordCountPropertyHandle->SetValue(0);
-			}
-
-			AddNewSupportedCultureComboButton->SetIsOpen(false);
-
-			NewEntryIndexToBeInitialized = INDEX_NONE;
-		}
-
-		ListedCultureStatisticProperties.Empty();
-		BuildListedCulturesList();
-
-		if (SupportedCultureListView.IsValid())
-		{
-			SupportedCultureListView->RequestListRefresh();
-		}
-	}
-
-	TSharedRef<ITableRow> FLocalizationTargetDetailCustomization::OnGenerateCultureRow(TSharedPtr<IPropertyHandle> CultureStatisticsPropertyHandle, const TSharedRef<STableViewBase>& Table)
-	{
-		return SNew(SLocalizationTargetEditorCultureRow, Table, DetailLayoutBuilder->GetPropertyUtilities(), TargetSettingsPropertyHandle.ToSharedRef(), CultureStatisticsPropertyHandle->GetIndexInArray());
-	}
-
-	bool FLocalizationTargetDetailCustomization::IsCultureSelectableAsSupported(FCulturePtr Culture)
-	{
-		// Can't select the native culture.
-		if (NativeCultureNamePropertyHandle.IsValid() && NativeCultureNamePropertyHandle->IsValidHandle())
-		{
-			FString NativeCultureName;
-			NativeCultureNamePropertyHandle->GetValue(NativeCultureName);
-			const FCulturePtr NativeCulture = FInternationalization::Get().GetCulture(NativeCultureName);
-			if (NativeCulture == Culture)
-			{
-				return false;
-			}
-		}
-
-		auto Is = [&](const TSharedPtr<IPropertyHandle>& SupportedCultureStatisticProperty)
-		{
-			// Can't select existing supported cultures.
-			const TSharedPtr<IPropertyHandle> SupportedCultureNamePropertyHandle = SupportedCultureStatisticProperty->GetChildHandle(GET_MEMBER_NAME_CHECKED(FCultureStatistics, CultureName));
-			if (SupportedCultureNamePropertyHandle.IsValid() && SupportedCultureNamePropertyHandle->IsValidHandle())
-			{
-				FString SupportedCultureName;
-				SupportedCultureNamePropertyHandle->GetValue(SupportedCultureName);
-				const FCulturePtr SupportedCulture = FInternationalization::Get().GetCulture(SupportedCultureName);
-				return SupportedCulture == Culture;
-			}
-
-			return false;
-		};
-
-		return !ListedCultureStatisticProperties.ContainsByPredicate(Is);
-	}
-
-	void FLocalizationTargetDetailCustomization::OnNewSupportedCultureSelected(FCulturePtr SelectedCulture, ESelectInfo::Type SelectInfo)
-	{
-		if(SupportedCulturesStatisticsPropertyHandle.IsValid() && SupportedCulturesStatisticsPropertyHandle->IsValidHandle())
-		{
-			uint32 NewElementIndex;
-			SupportedCulturesStatisticsPropertyHandle->AsArray()->GetNumElements(NewElementIndex);
-
-			// Add element, set info for later initialization.
-			SupportedCulturesStatisticsPropertyHandle->AsArray()->AddItem();
-			SelectedNewCulture = SelectedCulture;
-			NewEntryIndexToBeInitialized = NewElementIndex;
-
-			// Refresh UI.
-			if (SupportedCulturePicker.IsValid())
-			{
-				SupportedCulturePicker->RequestTreeRefresh();
-			}
-		}
-	}
-}
 
 void SLocalizationTargetEditor::Construct(const FArguments& InArgs, ULocalizationTargetSet* const InProjectSettings, ULocalizationTarget* const InLocalizationTarget)
 {
@@ -931,7 +9,6 @@ void SLocalizationTargetEditor::Construct(const FArguments& InArgs, ULocalizatio
 	FPropertyEditorModule& PropertyModule = FModuleManager::LoadModuleChecked<FPropertyEditorModule>("PropertyEditor");
 	FDetailsViewArgs DetailsViewArgs(false, false, false, FDetailsViewArgs::ENameAreaSettings::HideNameArea, false, nullptr, false, NAME_None);
 	TSharedRef<IDetailsView> DetailsView = PropertyModule.CreateDetailView(DetailsViewArgs);
-	DetailsView->RegisterInstancedCustomPropertyLayout(ULocalizationTarget::StaticClass(), FOnGetDetailCustomizationInstance::CreateLambda([&]()->TSharedRef<IDetailCustomization>{return MakeShareable(new FLocalizationTargetDetailCustomization(InProjectSettings, InLocalizationTarget));}));
 
 	ChildSlot
 		[
@@ -939,6 +16,4 @@ void SLocalizationTargetEditor::Construct(const FArguments& InArgs, ULocalizatio
 		];
 
 	DetailsView->SetObject(InLocalizationTarget, true);
-}
-
-#undef LOCTEXT_NAMESPACE
\ No newline at end of file
+}
\ No newline at end of file
