Updated copyright years on localization dashboard source files.

diff --git a/Engine/Source/Editor/LocalizationDashboard/Classes/LocalizationTarget.h b/Engine/Source/Editor/LocalizationDashboard/Classes/LocalizationTarget.h
--- a/Engine/Source/Editor/LocalizationDashboard/Classes/LocalizationTarget.h
+++ b/Engine/Source/Editor/LocalizationDashboard/Classes/LocalizationTarget.h
@@ -1,4 +1,4 @@
-// Copyright 1998-2014 Epic Games, Inc. All Rights Reserved.
+// Copyright 1998-2015 Epic Games, Inc. All Rights Reserved.
 
 #pragma once
 
diff --git a/Engine/Source/Editor/LocalizationDashboard/LocalizationDashboard.Build.cs b/Engine/Source/Editor/LocalizationDashboard/LocalizationDashboard.Build.cs
--- a/Engine/Source/Editor/LocalizationDashboard/LocalizationDashboard.Build.cs
+++ b/Engine/Source/Editor/LocalizationDashboard/LocalizationDashboard.Build.cs
@@ -1,4 +1,4 @@
-// Copyright 1998-2014 Epic Games, Inc. All Rights Reserved.
+// Copyright 1998-2015 Epic Games, Inc. All Rights Reserved.
 
 using UnrealBuildTool;
 
diff --git a/Engine/Source/Editor/LocalizationDashboard/Private/LocalizationCommandletExecution.cpp b/Engine/Source/Editor/LocalizationDashboard/Private/LocalizationCommandletExecution.cpp
--- a/Engine/Source/Editor/LocalizationDashboard/Private/LocalizationCommandletExecution.cpp
+++ b/Engine/Source/Editor/LocalizationDashboard/Private/LocalizationCommandletExecution.cpp
@@ -1,4 +1,4 @@
-﻿// Copyright 1998-2014 Epic Games, Inc. All Rights Reserved.
+﻿// Copyright 1998-2015 Epic Games, Inc. All Rights Reserved.
 
 #include "LocalizationDashboardPrivatePCH.h"
 #include "LocalizationCommandletExecution.h"
diff --git a/Engine/Source/Editor/LocalizationDashboard/Private/LocalizationCommandletExecution.h b/Engine/Source/Editor/LocalizationDashboard/Private/LocalizationCommandletExecution.h
--- a/Engine/Source/Editor/LocalizationDashboard/Private/LocalizationCommandletExecution.h
+++ b/Engine/Source/Editor/LocalizationDashboard/Private/LocalizationCommandletExecution.h
@@ -1,4 +1,4 @@
-// Copyright 1998-2014 Epic Games, Inc. All Rights Reserved.
+// Copyright 1998-2015 Epic Games, Inc. All Rights Reserved.
 #pragma once 
 
 #include "LocalizationConfigurationScript.h"
diff --git a/Engine/Source/Editor/LocalizationDashboard/Private/LocalizationCommandletTasks.cpp b/Engine/Source/Editor/LocalizationDashboard/Private/LocalizationCommandletTasks.cpp
--- a/Engine/Source/Editor/LocalizationDashboard/Private/LocalizationCommandletTasks.cpp
+++ b/Engine/Source/Editor/LocalizationDashboard/Private/LocalizationCommandletTasks.cpp
@@ -1,4 +1,4 @@
-// Copyright 1998-2014 Epic Games, Inc. All Rights Reserved.
+// Copyright 1998-2015 Epic Games, Inc. All Rights Reserved.
 
 #include "LocalizationDashboardPrivatePCH.h"
 #include "LocalizationCommandletTasks.h"
diff --git a/Engine/Source/Editor/LocalizationDashboard/Private/LocalizationCommandletTasks.h b/Engine/Source/Editor/LocalizationDashboard/Private/LocalizationCommandletTasks.h
--- a/Engine/Source/Editor/LocalizationDashboard/Private/LocalizationCommandletTasks.h
+++ b/Engine/Source/Editor/LocalizationDashboard/Private/LocalizationCommandletTasks.h
@@ -1,4 +1,4 @@
-// Copyright 1998-2014 Epic Games, Inc. All Rights Reserved.
+// Copyright 1998-2015 Epic Games, Inc. All Rights Reserved.
 #pragma once
 
 #include "Classes/LocalizationTarget.h"
diff --git a/Engine/Source/Editor/LocalizationDashboard/Private/LocalizationConfigurationScript.cpp b/Engine/Source/Editor/LocalizationDashboard/Private/LocalizationConfigurationScript.cpp
--- a/Engine/Source/Editor/LocalizationDashboard/Private/LocalizationConfigurationScript.cpp
+++ b/Engine/Source/Editor/LocalizationDashboard/Private/LocalizationConfigurationScript.cpp
@@ -1,4 +1,4 @@
-// Copyright 1998-2014 Epic Games, Inc. All Rights Reserved.
+// Copyright 1998-2015 Epic Games, Inc. All Rights Reserved.
 
 #include "LocalizationDashboardPrivatePCH.h"
 #include "LocalizationConfigurationScript.h"
diff --git a/Engine/Source/Editor/LocalizationDashboard/Private/LocalizationConfigurationScript.h b/Engine/Source/Editor/LocalizationDashboard/Private/LocalizationConfigurationScript.h
--- a/Engine/Source/Editor/LocalizationDashboard/Private/LocalizationConfigurationScript.h
+++ b/Engine/Source/Editor/LocalizationDashboard/Private/LocalizationConfigurationScript.h
@@ -1,4 +1,4 @@
-// Copyright 1998-2014 Epic Games, Inc. All Rights Reserved.
+// Copyright 1998-2015 Epic Games, Inc. All Rights Reserved.
 
 #pragma once
 
diff --git a/Engine/Source/Editor/LocalizationDashboard/Private/LocalizationDashboard.cpp b/Engine/Source/Editor/LocalizationDashboard/Private/LocalizationDashboard.cpp
--- a/Engine/Source/Editor/LocalizationDashboard/Private/LocalizationDashboard.cpp
+++ b/Engine/Source/Editor/LocalizationDashboard/Private/LocalizationDashboard.cpp
@@ -1,4 +1,4 @@
-﻿// Copyright 1998-2014 Epic Games, Inc. All Rights Reserved.
+﻿// Copyright 1998-2015 Epic Games, Inc. All Rights Reserved.
 
 #include "LocalizationDashboardPrivatePCH.h"
 #include "LocalizationDashboard.h"
diff --git a/Engine/Source/Editor/LocalizationDashboard/Private/LocalizationDashboard.h b/Engine/Source/Editor/LocalizationDashboard/Private/LocalizationDashboard.h
--- a/Engine/Source/Editor/LocalizationDashboard/Private/LocalizationDashboard.h
+++ b/Engine/Source/Editor/LocalizationDashboard/Private/LocalizationDashboard.h
@@ -1,4 +1,4 @@
-// Copyright 1998-2014 Epic Games, Inc. All Rights Reserved.
+// Copyright 1998-2015 Epic Games, Inc. All Rights Reserved.
 
 #pragma once
 
diff --git a/Engine/Source/Editor/LocalizationDashboard/Private/LocalizationDashboardModule.cpp b/Engine/Source/Editor/LocalizationDashboard/Private/LocalizationDashboardModule.cpp
--- a/Engine/Source/Editor/LocalizationDashboard/Private/LocalizationDashboardModule.cpp
+++ b/Engine/Source/Editor/LocalizationDashboard/Private/LocalizationDashboardModule.cpp
@@ -1,4 +1,4 @@
-// Copyright 1998-2014 Epic Games, Inc. All Rights Reserved.
+// Copyright 1998-2015 Epic Games, Inc. All Rights Reserved.
 
 #include "LocalizationDashboardPrivatePCH.h"
 #include "ILocalizationDashboardModule.h"
diff --git a/Engine/Source/Editor/LocalizationDashboard/Private/LocalizationDashboardPrivatePCH.h b/Engine/Source/Editor/LocalizationDashboard/Private/LocalizationDashboardPrivatePCH.h
--- a/Engine/Source/Editor/LocalizationDashboard/Private/LocalizationDashboardPrivatePCH.h
+++ b/Engine/Source/Editor/LocalizationDashboard/Private/LocalizationDashboardPrivatePCH.h
@@ -1,4 +1,4 @@
-// Copyright 1998-2014 Epic Games, Inc. All Rights Reserved.
+// Copyright 1998-2015 Epic Games, Inc. All Rights Reserved.
 
 #pragma once
 
diff --git a/Engine/Source/Editor/LocalizationDashboard/Private/LocalizationTarget.cpp b/Engine/Source/Editor/LocalizationDashboard/Private/LocalizationTarget.cpp
--- a/Engine/Source/Editor/LocalizationDashboard/Private/LocalizationTarget.cpp
+++ b/Engine/Source/Editor/LocalizationDashboard/Private/LocalizationTarget.cpp
@@ -1,4 +1,4 @@
-// Copyright 1998-2014 Epic Games, Inc. All Rights Reserved.
+// Copyright 1998-2015 Epic Games, Inc. All Rights Reserved.
 
 #include "LocalizationDashboardPrivatePCH.h"
 #include "Classes/LocalizationTarget.h"
diff --git a/Engine/Source/Editor/LocalizationDashboard/Private/SCulturePicker.cpp b/Engine/Source/Editor/LocalizationDashboard/Private/SCulturePicker.cpp
--- a/Engine/Source/Editor/LocalizationDashboard/Private/SCulturePicker.cpp
+++ b/Engine/Source/Editor/LocalizationDashboard/Private/SCulturePicker.cpp
@@ -1,4 +1,4 @@
-// Copyright 1998-2014 Epic Games, Inc. All Rights Reserved.
+// Copyright 1998-2015 Epic Games, Inc. All Rights Reserved.
 
 #include "LocalizationDashboardPrivatePCH.h"
 #include "SCulturePicker.h"
diff --git a/Engine/Source/Editor/LocalizationDashboard/Private/SCulturePicker.h b/Engine/Source/Editor/LocalizationDashboard/Private/SCulturePicker.h
--- a/Engine/Source/Editor/LocalizationDashboard/Private/SCulturePicker.h
+++ b/Engine/Source/Editor/LocalizationDashboard/Private/SCulturePicker.h
@@ -1,4 +1,4 @@
-// Copyright 1998-2014 Epic Games, Inc. All Rights Reserved.
+// Copyright 1998-2015 Epic Games, Inc. All Rights Reserved.
 
 #pragma once
 
diff --git a/Engine/Source/Editor/LocalizationDashboard/Private/SLocalizationDashboardTargetRow.cpp b/Engine/Source/Editor/LocalizationDashboard/Private/SLocalizationDashboardTargetRow.cpp
--- a/Engine/Source/Editor/LocalizationDashboard/Private/SLocalizationDashboardTargetRow.cpp
+++ b/Engine/Source/Editor/LocalizationDashboard/Private/SLocalizationDashboardTargetRow.cpp
@@ -1,4 +1,4 @@
-// Copyright 1998-2014 Epic Games, Inc. All Rights Reserved.
+// Copyright 1998-2015 Epic Games, Inc. All Rights Reserved.
 
 #include "LocalizationDashboardPrivatePCH.h"
 #include "LocalizationDashboard.h"
diff --git a/Engine/Source/Editor/LocalizationDashboard/Private/SLocalizationDashboardTargetRow.h b/Engine/Source/Editor/LocalizationDashboard/Private/SLocalizationDashboardTargetRow.h
--- a/Engine/Source/Editor/LocalizationDashboard/Private/SLocalizationDashboardTargetRow.h
+++ b/Engine/Source/Editor/LocalizationDashboard/Private/SLocalizationDashboardTargetRow.h
@@ -1,4 +1,4 @@
-// Copyright 1998-2014 Epic Games, Inc. All Rights Reserved.
+// Copyright 1998-2015 Epic Games, Inc. All Rights Reserved.
 #pragma once
 
 #include "STableRow.h"
diff --git a/Engine/Source/Editor/LocalizationDashboard/Private/SLocalizationTargetEditor.cpp b/Engine/Source/Editor/LocalizationDashboard/Private/SLocalizationTargetEditor.cpp
--- a/Engine/Source/Editor/LocalizationDashboard/Private/SLocalizationTargetEditor.cpp
+++ b/Engine/Source/Editor/LocalizationDashboard/Private/SLocalizationTargetEditor.cpp
@@ -1,4 +1,4 @@
-// Copyright 1998-2014 Epic Games, Inc. All Rights Reserved.
+// Copyright 1998-2015 Epic Games, Inc. All Rights Reserved.
 
 #include "LocalizationDashboardPrivatePCH.h"
 
diff --git a/Engine/Source/Editor/LocalizationDashboard/Private/SLocalizationTargetEditor.h b/Engine/Source/Editor/LocalizationDashboard/Private/SLocalizationTargetEditor.h
--- a/Engine/Source/Editor/LocalizationDashboard/Private/SLocalizationTargetEditor.h
+++ b/Engine/Source/Editor/LocalizationDashboard/Private/SLocalizationTargetEditor.h
@@ -1,4 +1,4 @@
-// Copyright 1998-2014 Epic Games, Inc. All Rights Reserved.
+// Copyright 1998-2015 Epic Games, Inc. All Rights Reserved.
 
 #pragma once
 
diff --git a/Engine/Source/Editor/LocalizationDashboard/Private/SLocalizationTargetEditorCultureRow.cpp b/Engine/Source/Editor/LocalizationDashboard/Private/SLocalizationTargetEditorCultureRow.cpp
--- a/Engine/Source/Editor/LocalizationDashboard/Private/SLocalizationTargetEditorCultureRow.cpp
+++ b/Engine/Source/Editor/LocalizationDashboard/Private/SLocalizationTargetEditorCultureRow.cpp
@@ -1,4 +1,4 @@
-// Copyright 1998-2014 Epic Games, Inc. All Rights Reserved.
+// Copyright 1998-2015 Epic Games, Inc. All Rights Reserved.
 
 #include "LocalizationDashboardPrivatePCH.h"
 #include "SLocalizationTargetEditorCultureRow.h"
diff --git a/Engine/Source/Editor/LocalizationDashboard/Private/SLocalizationTargetEditorCultureRow.h b/Engine/Source/Editor/LocalizationDashboard/Private/SLocalizationTargetEditorCultureRow.h
--- a/Engine/Source/Editor/LocalizationDashboard/Private/SLocalizationTargetEditorCultureRow.h
+++ b/Engine/Source/Editor/LocalizationDashboard/Private/SLocalizationTargetEditorCultureRow.h
@@ -1,4 +1,4 @@
-// Copyright 1998-2014 Epic Games, Inc. All Rights Reserved.
+// Copyright 1998-2015 Epic Games, Inc. All Rights Reserved.
 #pragma once
 
 #include "STableRow.h"
diff --git a/Engine/Source/Editor/LocalizationDashboard/Private/SLocalizationTargetStatusButton.cpp b/Engine/Source/Editor/LocalizationDashboard/Private/SLocalizationTargetStatusButton.cpp
--- a/Engine/Source/Editor/LocalizationDashboard/Private/SLocalizationTargetStatusButton.cpp
+++ b/Engine/Source/Editor/LocalizationDashboard/Private/SLocalizationTargetStatusButton.cpp
@@ -1,4 +1,4 @@
-// Copyright 1998-2014 Epic Games, Inc. All Rights Reserved.
+// Copyright 1998-2015 Epic Games, Inc. All Rights Reserved.
 
 #include "LocalizationDashboardPrivatePCH.h"
 #include "SLocalizationTargetStatusButton.h"
diff --git a/Engine/Source/Editor/LocalizationDashboard/Private/SLocalizationTargetStatusButton.h b/Engine/Source/Editor/LocalizationDashboard/Private/SLocalizationTargetStatusButton.h
--- a/Engine/Source/Editor/LocalizationDashboard/Private/SLocalizationTargetStatusButton.h
+++ b/Engine/Source/Editor/LocalizationDashboard/Private/SLocalizationTargetStatusButton.h
@@ -1,4 +1,4 @@
-// Copyright 1998-2014 Epic Games, Inc. All Rights Reserved.
+// Copyright 1998-2015 Epic Games, Inc. All Rights Reserved.
 #pragma once
 
 #include "SCompoundWidget.h"
diff --git a/Engine/Source/Editor/LocalizationDashboard/Private/TargetsTableEntry.h b/Engine/Source/Editor/LocalizationDashboard/Private/TargetsTableEntry.h
--- a/Engine/Source/Editor/LocalizationDashboard/Private/TargetsTableEntry.h
+++ b/Engine/Source/Editor/LocalizationDashboard/Private/TargetsTableEntry.h
@@ -1,4 +1,4 @@
-// Copyright 1998-2014 Epic Games, Inc. All Rights Reserved.
+// Copyright 1998-2015 Epic Games, Inc. All Rights Reserved.
 #pragma once
 
 class ULocalizationTarget;
diff --git a/Engine/Source/Editor/LocalizationDashboard/Public/ILocalizationDashboardModule.h b/Engine/Source/Editor/LocalizationDashboard/Public/ILocalizationDashboardModule.h
--- a/Engine/Source/Editor/LocalizationDashboard/Public/ILocalizationDashboardModule.h
+++ b/Engine/Source/Editor/LocalizationDashboard/Public/ILocalizationDashboardModule.h
@@ -1,4 +1,4 @@
-// Copyright 1998-2014 Epic Games, Inc. All Rights Reserved.
+// Copyright 1998-2015 Epic Games, Inc. All Rights Reserved.
 
 #pragma once
 
diff --git a/Engine/Source/Editor/LocalizationDashboard/Public/ILocalizationServiceProvider.h b/Engine/Source/Editor/LocalizationDashboard/Public/ILocalizationServiceProvider.h
--- a/Engine/Source/Editor/LocalizationDashboard/Public/ILocalizationServiceProvider.h
+++ b/Engine/Source/Editor/LocalizationDashboard/Public/ILocalizationServiceProvider.h
@@ -1,4 +1,4 @@
-// Copyright 1998-2014 Epic Games, Inc. All Rights Reserved.
+// Copyright 1998-2015 Epic Games, Inc. All Rights Reserved.
 
 #pragma once
 
