Fixed non-Unity build errors

diff --git a/Engine/Source/Editor/LocalizationDashboard/Private/LocalizationDashboard.cpp b/Engine/Source/Editor/LocalizationDashboard/Private/LocalizationDashboard.cpp
--- a/Engine/Source/Editor/LocalizationDashboard/Private/LocalizationDashboard.cpp
+++ b/Engine/Source/Editor/LocalizationDashboard/Private/LocalizationDashboard.cpp
@@ -8,6 +8,7 @@
 #include "IMainFrameModule.h"
 #include "ModuleManager.h"
 #include "SLocalizationTargetEditor.h"
+#include "LocalizationTarget.h"
 
 #define LOCTEXT_NAMESPACE "LocalizationDashboard"
 
diff --git a/Engine/Source/Editor/LocalizationDashboard/Private/LocalizationTargetDetailCustomization.cpp b/Engine/Source/Editor/LocalizationDashboard/Private/LocalizationTargetDetailCustomization.cpp
--- a/Engine/Source/Editor/LocalizationDashboard/Private/LocalizationTargetDetailCustomization.cpp
+++ b/Engine/Source/Editor/LocalizationDashboard/Private/LocalizationTargetDetailCustomization.cpp
@@ -9,6 +9,9 @@
 #include "SLocalizationTargetEditorCultureRow.h"
 #include "SCulturePicker.h"
 #include "IDetailPropertyRow.h"
+#include "DetailCategoryBuilder.h"
+#include "IDetailsView.h"
+#include "LocalizationCommandletTasks.h"
 
 #define LOCTEXT_NAMESPACE "LocalizationTargetEditor"
 
diff --git a/Engine/Source/Editor/LocalizationDashboard/Private/LocalizationTargetDetailCustomization.h b/Engine/Source/Editor/LocalizationDashboard/Private/LocalizationTargetDetailCustomization.h
--- a/Engine/Source/Editor/LocalizationDashboard/Private/LocalizationTargetDetailCustomization.h
+++ b/Engine/Source/Editor/LocalizationDashboard/Private/LocalizationTargetDetailCustomization.h
@@ -2,9 +2,8 @@
 
 #pragma once
 #include "IDetailCustomization.h"
+#include "LocalizationTarget.h"
 
-class ULocalizationTargetSet;
-class ULocalizationTarget;
 class SCulturePicker;
 
 class FLocalizationTargetDetailCustomization : public IDetailCustomization
diff --git a/Engine/Source/Editor/LocalizationDashboard/Private/LocalizationTargetSetDetailCustomization.cpp b/Engine/Source/Editor/LocalizationDashboard/Private/LocalizationTargetSetDetailCustomization.cpp
--- a/Engine/Source/Editor/LocalizationDashboard/Private/LocalizationTargetSetDetailCustomization.cpp
+++ b/Engine/Source/Editor/LocalizationDashboard/Private/LocalizationTargetSetDetailCustomization.cpp
@@ -6,6 +6,9 @@
 #include "LocalizationTarget.h"
 #include "DetailWidgetRow.h"
 #include "SLocalizationDashboardTargetRow.h"
+#include "DetailCategoryBuilder.h"
+#include "IDetailsView.h"
+#include "LocalizationCommandletTasks.h"
 
 #define LOCTEXT_NAMESPACE "LocalizationDashboard"
 
diff --git a/Engine/Source/Editor/LocalizationDashboard/Private/LocalizationTargetSetDetailCustomization.h b/Engine/Source/Editor/LocalizationDashboard/Private/LocalizationTargetSetDetailCustomization.h
--- a/Engine/Source/Editor/LocalizationDashboard/Private/LocalizationTargetSetDetailCustomization.h
+++ b/Engine/Source/Editor/LocalizationDashboard/Private/LocalizationTargetSetDetailCustomization.h
@@ -2,6 +2,12 @@
 
 #pragma once
 
+#include "IDetailCustomization.h"
+
+class ULocalizationTarget;
+class IPropertyHandle;
+class ULocalizationTargetSet;
+
 //class ILocalizationServiceProvider;
 
 //struct FLocalizationServiceProviderWrapper
diff --git a/Engine/Source/Editor/LocalizationDashboard/Private/SLocalizationTargetEditor.cpp b/Engine/Source/Editor/LocalizationDashboard/Private/SLocalizationTargetEditor.cpp
--- a/Engine/Source/Editor/LocalizationDashboard/Private/SLocalizationTargetEditor.cpp
+++ b/Engine/Source/Editor/LocalizationDashboard/Private/SLocalizationTargetEditor.cpp
@@ -1,6 +1,10 @@
 // Copyright 1998-2015 Epic Games, Inc. All Rights Reserved.
 
 #include "LocalizationDashboardPrivatePCH.h"
+#include "SLocalizationTargetEditor.h"
+#include "IDetailsView.h"
+#include "LocalizationTarget.h"
+#include "PropertyEditorModule.h"
 
 void SLocalizationTargetEditor::Construct(const FArguments& InArgs, ULocalizationTargetSet* const InProjectSettings, ULocalizationTarget* const InLocalizationTarget)
 {
