One more forward declaration to fix non-unity

diff --git a/Engine/Source/Editor/LocalizationDashboard/Private/LocalizationTargetDetailCustomization.h b/Engine/Source/Editor/LocalizationDashboard/Private/LocalizationTargetDetailCustomization.h
--- a/Engine/Source/Editor/LocalizationDashboard/Private/LocalizationTargetDetailCustomization.h
+++ b/Engine/Source/Editor/LocalizationDashboard/Private/LocalizationTargetDetailCustomization.h
@@ -5,6 +5,7 @@
 #include "LocalizationTarget.h"
 
 class SCulturePicker;
+class IPropertyHandle;
 
 class FLocalizationTargetDetailCustomization : public IDetailCustomization
 {
