Status buttons on the localization dashboard/localization target editor UI now use localization commandlet executor UI to show progress or allow cancellation.

diff --git a/Engine/Source/Editor/LocalizationDashboard/Private/SLocalizationTargetStatusButton.cpp b/Engine/Source/Editor/LocalizationDashboard/Private/SLocalizationTargetStatusButton.cpp
--- a/Engine/Source/Editor/LocalizationDashboard/Private/SLocalizationTargetStatusButton.cpp
+++ b/Engine/Source/Editor/LocalizationDashboard/Private/SLocalizationTargetStatusButton.cpp
@@ -85,12 +85,9 @@ FReply SLocalizationTargetStatusButton::OnClicked()
 	case ELocalizationTargetStatus::Unknown:
 		{
 			// Generate conflict report.
-			const FString ReportScriptPath = LocalizationConfigurationScript::GetReportScriptPath(Target->Settings);
-			LocalizationConfigurationScript::GenerateReportScript(Target->Settings).Write(ReportScriptPath);
-			TSharedPtr<FLocalizationCommandletProcess> ReportProcess = FLocalizationCommandletProcess::Execute(ReportScriptPath);
-			if (ReportProcess.IsValid())
+			const TSharedPtr<SWindow> ParentWindow = FSlateApplication::Get().FindWidgetWindow(AsShared());
+			if (LocalizationCommandletTasks::GenerateReportsForTarget(ParentWindow.ToSharedRef(), Target->Settings))
 			{
-				FPlatformProcess::WaitForProc(ReportProcess->GetHandle());
 				Target->Settings.UpdateStatusFromConflictReport();
 			}
 		}
