Refactored localization dashboard and localization target editor details customization to make it easier to add new members to the classes being customized, in preparation for adding gather step customization.

diff --git a/Engine/Source/Editor/LocalizationDashboard/Private/LocalizationTargetDetailCustomization.cpp b/Engine/Source/Editor/LocalizationDashboard/Private/LocalizationTargetDetailCustomization.cpp
--- a/Engine/Source/Editor/LocalizationDashboard/Private/LocalizationTargetDetailCustomization.cpp
+++ b/Engine/Source/Editor/LocalizationDashboard/Private/LocalizationTargetDetailCustomization.cpp
@@ -12,6 +12,7 @@
 #include "DetailCategoryBuilder.h"
 #include "IDetailsView.h"
 #include "LocalizationCommandletTasks.h"
+#include "ObjectEditorUtils.h"
 
 #define LOCTEXT_NAMESPACE "LocalizationTargetEditor"
 
@@ -65,249 +66,258 @@ void FLocalizationTargetDetailCustomization::CustomizeDetails(IDetailLayoutBuild
 		TargetSettingsPropertyHandle->MarkHiddenByCustomization();
 	}
 
-	//{
-	//	IDetailCategoryBuilder& ServiceProviderCategoryBuilder = DetailBuilder.EditCategory(FName("LocalizationServiceProvider"));
-
-	//	ILocalizationServiceProvider* const LSP = ILocalizationDashboardModule::Get().GetCurrentLocalizationServiceProvider();
-	//	if (LSP && LocalizationTarget)
-	//	{
-	//		LSP->CustomizeTargetDetails(ServiceProviderCategoryBuilder, *LocalizationTarget);
-	//	}
-	//}
-
-	if (TargetSettingsPropertyHandle.IsValid() && TargetSettingsPropertyHandle->IsValidHandle())
-	{
-		IDetailCategoryBuilder& DetailCategoryBuilder = DetailBuilder.EditCategory(FName("Target"));
-
-		const TSharedPtr<IPropertyHandle> NamePropertyHandle = TargetSettingsPropertyHandle->GetChildHandle(GET_MEMBER_NAME_CHECKED(FLocalizationTargetSettings, Name));
-		if (NamePropertyHandle.IsValid() && NamePropertyHandle->IsValidHandle())
+	typedef TFunction<void (const TSharedRef<IPropertyHandle>&, IDetailCategoryBuilder&)> FPropertyCustomizationFunction;
+	TMap<FName, FPropertyCustomizationFunction> PropertyCustomizationMap;
+
+	PropertyCustomizationMap.Add(GET_MEMBER_NAME_CHECKED(FLocalizationTargetSettings, Name), [&](const TSharedRef<IPropertyHandle>& MemberPropertyHandle, IDetailCategoryBuilder& DetailCategoryBuilder)
+	{
+		FDetailWidgetRow& DetailWidgetRow = DetailCategoryBuilder.AddCustomRow( MemberPropertyHandle->GetPropertyDisplayName() );
+		DetailWidgetRow.NameContent()
+			[
+				MemberPropertyHandle->CreatePropertyNameWidget()
+			];
+		DetailWidgetRow.ValueContent()
+			[
+				SAssignNew(TargetNameEditableTextBox, SEditableTextBox)
+				.Text(this, &FLocalizationTargetDetailCustomization::GetTargetName)
+				.RevertTextOnEscape(true)
+				.OnTextChanged(this, &FLocalizationTargetDetailCustomization::OnTargetNameChanged)
+				.OnTextCommitted(this, &FLocalizationTargetDetailCustomization::OnTargetNameCommitted)
+			];
+	});
+
+	PropertyCustomizationMap.Add(GET_MEMBER_NAME_CHECKED(FLocalizationTargetSettings, Status), [&](const TSharedRef<IPropertyHandle>& MemberPropertyHandle, IDetailCategoryBuilder& DetailCategoryBuilder)
+	{
+		FDetailWidgetRow& StatusRow = DetailCategoryBuilder.AddCustomRow( MemberPropertyHandle->GetPropertyDisplayName() );
+		StatusRow.NameContent()
+			[
+				MemberPropertyHandle->CreatePropertyNameWidget()
+			];
+		StatusRow.ValueContent()
+			.HAlign(HAlign_Left)
+			.VAlign(VAlign_Center)
+			[
+				SNew(SLocalizationTargetStatusButton, *LocalizationTarget)
+			];
+	});
+
+	PropertyCustomizationMap.Add(GET_MEMBER_NAME_CHECKED(FLocalizationTargetSettings, TargetDependencies), [&](const TSharedRef<IPropertyHandle>& MemberPropertyHandle, IDetailCategoryBuilder& DetailCategoryBuilder)
+	{
+		const auto& MenuContentLambda = [this]() -> TSharedRef<SWidget>
 		{
-			FDetailWidgetRow& DetailWidgetRow = DetailCategoryBuilder.AddCustomRow( LOCTEXT("TargetNameFilterString", "Target Name") );
-			DetailWidgetRow.NameContent()
-				[
-					NamePropertyHandle->CreatePropertyNameWidget()
-				];
-			DetailWidgetRow.ValueContent()
-				[
-					SAssignNew(TargetNameEditableTextBox, SEditableTextBox)
-					.Text(this, &FLocalizationTargetDetailCustomization::GetTargetName)
-					.RevertTextOnEscape(true)
-					.OnTextChanged(this, &FLocalizationTargetDetailCustomization::OnTargetNameChanged)
-					.OnTextCommitted(this, &FLocalizationTargetDetailCustomization::OnTargetNameCommitted)
-				];
-		}
-		const TSharedPtr<IPropertyHandle> StatusPropertyHandle = TargetSettingsPropertyHandle->GetChildHandle(GET_MEMBER_NAME_CHECKED(FLocalizationTargetSettings, Status));
-		if (StatusPropertyHandle.IsValid() && StatusPropertyHandle->IsValidHandle())
-		{
-			StatusPropertyHandle->MarkHiddenByCustomization();
-			if (LocalizationTarget.IsValid())
+			RebuildTargetsList();
+
+			if (TargetDependenciesOptionsList.Num() > 0)
 			{
-				FDetailWidgetRow& StatusRow = DetailCategoryBuilder.AddCustomRow( LOCTEXT("ConflictReportStatusFilterString", "Conflict Report Status") );
-				StatusRow.NameContent()
-					[
-						StatusPropertyHandle->CreatePropertyNameWidget()
-					];
-				StatusRow.ValueContent()
-					.HAlign(HAlign_Left)
-					.VAlign(VAlign_Center)
+				return SNew(SBox)
+					.MaxDesiredHeight(400.0f)
+					.MaxDesiredWidth(300.0f)
 					[
-						SNew(SLocalizationTargetStatusButton, *LocalizationTarget)
+						SNew(SListView<ULocalizationTarget*>)
+						.SelectionMode(ESelectionMode::None)
+						.ListItemsSource(&TargetDependenciesOptionsList)
+						.OnGenerateRow(this, &FLocalizationTargetDetailCustomization::OnGenerateTargetRow)
 					];
 			}
-		}
-
-		const TSharedPtr<IPropertyHandle> TargetDependenciesPropertyHandle = TargetSettingsPropertyHandle->GetChildHandle(GET_MEMBER_NAME_CHECKED(FLocalizationTargetSettings, TargetDependencies));
-		if (TargetDependenciesPropertyHandle.IsValid() && TargetDependenciesPropertyHandle->IsValidHandle())
-		{
-			TargetDependenciesPropertyHandle->MarkHiddenByCustomization();
-			if (LocalizationTarget.IsValid())
+			else
 			{
-				FDetailWidgetRow& TargetDependenciesRow = DetailCategoryBuilder.AddCustomRow( LOCTEXT("ConflictReportStatusFilterString", "Conflict Report Status") );
-				TargetDependenciesRow.NameContent()
-					[
-						TargetDependenciesPropertyHandle->CreatePropertyNameWidget()
-					];
-				TargetDependenciesRow.ValueContent()
-					.HAlign(HAlign_Left)
-					.VAlign(VAlign_Center)
-					[
-						SNew(SComboButton)
-						.ButtonContent()
-						[
-							SAssignNew(TargetDependenciesHorizontalBox, SHorizontalBox)
-						]
-						.HasDownArrow(true)
-							.OnGetMenuContent_Lambda([this]() -> TSharedRef<SWidget>
-						{
-							RebuildTargetsList();
-
-							if (TargetDependenciesOptionsList.Num() > 0)
-							{
-								return SNew(SBox)
-									.MaxDesiredHeight(400.0f)
-									.MaxDesiredWidth(300.0f)
-									[
-										SNew(SListView<ULocalizationTarget*>)
-										.SelectionMode(ESelectionMode::None)
-										.ListItemsSource(&TargetDependenciesOptionsList)
-										.OnGenerateRow(this, &FLocalizationTargetDetailCustomization::OnGenerateTargetRow)
-									];
-							}
-							else
-							{
-								return SNullWidget::NullWidget;
-							}
-						})
-					];
-
-				RebuildTargetDependenciesBox();
+				return SNullWidget::NullWidget;
 			}
-		}
+		};
+
+		FDetailWidgetRow& TargetDependenciesRow = DetailCategoryBuilder.AddCustomRow( MemberPropertyHandle->GetPropertyDisplayName() );
+		TargetDependenciesRow.NameContent()
+			[
+				MemberPropertyHandle->CreatePropertyNameWidget()
+			];
+		TargetDependenciesRow.ValueContent()
+			.HAlign(HAlign_Left)
+			.VAlign(VAlign_Center)
+			[
+				SNew(SComboButton)
+				.ButtonContent()
+				[
+					SAssignNew(TargetDependenciesHorizontalBox, SHorizontalBox)
+				]
+				.HasDownArrow(true)
+				.OnGetMenuContent_Lambda(MenuContentLambda)
+			];
 
-		const TSharedPtr<IPropertyHandle> AdditionalManifestDependenciesPropertyHandle = TargetSettingsPropertyHandle->GetChildHandle(GET_MEMBER_NAME_CHECKED(FLocalizationTargetSettings, AdditionalManifestDependencies));
-		if (AdditionalManifestDependenciesPropertyHandle.IsValid() && AdditionalManifestDependenciesPropertyHandle->IsValidHandle())
-		{
-			DetailCategoryBuilder.AddProperty(AdditionalManifestDependenciesPropertyHandle);
-		}
-	}
+		RebuildTargetDependenciesBox();
+	});
 
-	if (TargetSettingsPropertyHandle.IsValid() && TargetSettingsPropertyHandle->IsValidHandle())
+	PropertyCustomizationMap.Add(GET_MEMBER_NAME_CHECKED(FLocalizationTargetSettings, NativeCultureStatistics), [&](const TSharedRef<IPropertyHandle>& MemberPropertyHandle, IDetailCategoryBuilder& DetailCategoryBuilder)
 	{
-		IDetailCategoryBuilder& DetailCategoryBuilder = DetailBuilder.EditCategory(FName("Cultures"));
+		NativeCultureStatisticsPropertyHandle = MemberPropertyHandle;
+		NativeCultureNamePropertyHandle = NativeCultureStatisticsPropertyHandle->GetChildHandle(GET_MEMBER_NAME_CHECKED(FCultureStatistics, CultureName));
 
-		NativeCultureStatisticsPropertyHandle = TargetSettingsPropertyHandle->GetChildHandle(GET_MEMBER_NAME_CHECKED(FLocalizationTargetSettings, NativeCultureStatistics));
-		if (NativeCultureStatisticsPropertyHandle.IsValid() && NativeCultureStatisticsPropertyHandle->IsValidHandle())
+		if (NativeCultureNamePropertyHandle.IsValid() && NativeCultureNamePropertyHandle->IsValidHandle())
 		{
-			NativeCultureNamePropertyHandle = NativeCultureStatisticsPropertyHandle->GetChildHandle(GET_MEMBER_NAME_CHECKED(FCultureStatistics, CultureName));
-			if (NativeCultureNamePropertyHandle.IsValid() && NativeCultureNamePropertyHandle->IsValidHandle())
-			{
-				NativeCultureNamePropertyHandle->MarkHiddenByCustomization();
-				FDetailWidgetRow& NativeCultureRow = DetailCategoryBuilder.AddCustomRow( LOCTEXT("NativeCultureNameFilterString", "Native Culture Name") );
+			NativeCultureNamePropertyHandle->MarkHiddenByCustomization();
+			FDetailWidgetRow& NativeCultureRow = DetailCategoryBuilder.AddCustomRow( NativeCultureNamePropertyHandle->GetPropertyDisplayName() );
 
-				FString NativeCultureName;
-				NativeCultureNamePropertyHandle->GetValue(NativeCultureName);
-				FCulturePtr NativeCulture = FInternationalization::Get().GetCulture(NativeCultureName);
+			FString NativeCultureName;
+			NativeCultureNamePropertyHandle->GetValue(NativeCultureName);
+			FCulturePtr NativeCulture = FInternationalization::Get().GetCulture(NativeCultureName);
 
-				NativeCultureRow.NameContent()
-					[
-						NativeCultureNamePropertyHandle->CreatePropertyNameWidget(LOCTEXT("NativeCultureNameLabel", "Native Culture"))
-					];
+			NativeCultureRow.NameContent()
+				[
+					NativeCultureNamePropertyHandle->CreatePropertyNameWidget(LOCTEXT("NativeCultureNameLabel", "Native Culture"))
+				];
 
-				NativeCultureRow.ValueContent()
+			NativeCultureRow.ValueContent()
+				[
+					SAssignNew(NativeCultureComboButton, SComboButton)
+					.ButtonContent()
 					[
-						SAssignNew(NativeCultureComboButton, SComboButton)
-						.ButtonContent()
+						SNew(STextBlock)
+						.Text(this, &FLocalizationTargetDetailCustomization::GetNativeCultureDisplayName)
+						.ToolTipText(this, &FLocalizationTargetDetailCustomization::GetNativeCultureName)
+					]
+					.HasDownArrow(true)
+						.MenuContent()
 						[
-							SNew(STextBlock)
-							.Text(this, &FLocalizationTargetDetailCustomization::GetNativeCultureDisplayName)
-							.ToolTipText(this, &FLocalizationTargetDetailCustomization::GetNativeCultureName)
-						]
-						.HasDownArrow(true)
-							.MenuContent()
+							SNew(SBox)
+							.MaxDesiredHeight(400.0f)
+							.MaxDesiredWidth(300.0f)
 							[
-								SNew(SBox)
-								.MaxDesiredHeight(400.0f)
-								.MaxDesiredWidth(300.0f)
-								[
-									SNew(SCulturePicker)
-									.OnSelectionChanged(this, &FLocalizationTargetDetailCustomization::OnNativeCultureSelected)
-									.InitialSelection(NativeCulture)
-								]
+								SNew(SCulturePicker)
+								.OnSelectionChanged(this, &FLocalizationTargetDetailCustomization::OnNativeCultureSelected)
+								.InitialSelection(NativeCulture)
 							]
-					];
-			}
+						]
+				];
 		}
+	});
 
-		SupportedCulturesStatisticsPropertyHandle = TargetSettingsPropertyHandle->GetChildHandle(GET_MEMBER_NAME_CHECKED(FLocalizationTargetSettings, SupportedCulturesStatistics));
-		if (SupportedCulturesStatisticsPropertyHandle.IsValid() && SupportedCulturesStatisticsPropertyHandle->IsValidHandle())
-		{
-			SupportedCulturesStatisticsPropertyHandle->MarkHiddenByCustomization();
-			SupportedCulturesStatisticsPropertyHandle_OnNumElementsChanged = FSimpleDelegate::CreateSP(this, &FLocalizationTargetDetailCustomization::RebuildListedCulturesList);
-			SupportedCulturesStatisticsPropertyHandle->AsArray()->SetOnNumElementsChanged(SupportedCulturesStatisticsPropertyHandle_OnNumElementsChanged);
+	PropertyCustomizationMap.Add(GET_MEMBER_NAME_CHECKED(FLocalizationTargetSettings, SupportedCulturesStatistics), [&](const TSharedRef<IPropertyHandle>& MemberPropertyHandle, IDetailCategoryBuilder& DetailCategoryBuilder)
+	{
+		SupportedCulturesStatisticsPropertyHandle = MemberPropertyHandle;
 
-			FLocalizationTargetEditorCommands::Register();
-			const TSharedRef< FUICommandList > CommandList = MakeShareable(new FUICommandList);
-			FToolBarBuilder ToolBarBuilder( CommandList, FMultiBoxCustomization::AllowCustomization("LocalizationTargetEditor") );
+		SupportedCulturesStatisticsPropertyHandle_OnNumElementsChanged = FSimpleDelegate::CreateSP(this, &FLocalizationTargetDetailCustomization::RebuildListedCulturesList);
+		SupportedCulturesStatisticsPropertyHandle->AsArray()->SetOnNumElementsChanged(SupportedCulturesStatisticsPropertyHandle_OnNumElementsChanged);
 
-			CommandList->MapAction( FLocalizationTargetEditorCommands::Get().Gather, FExecuteAction::CreateSP(this, &FLocalizationTargetDetailCustomization::Gather));
-			ToolBarBuilder.AddToolBarButton(FLocalizationTargetEditorCommands::Get().Gather, NAME_None, TAttribute<FText>(), TAttribute<FText>(), FSlateIcon(FEditorStyle::GetStyleSetName(), "LocalizationTargetEditor.Gather"));
+		FLocalizationTargetEditorCommands::Register();
+		const TSharedRef< FUICommandList > CommandList = MakeShareable(new FUICommandList);
+		FToolBarBuilder ToolBarBuilder( CommandList, FMultiBoxCustomization::AllowCustomization("LocalizationTargetEditor") );
 
-			CommandList->MapAction( FLocalizationTargetEditorCommands::Get().ImportAllCultures, FExecuteAction::CreateSP(this, &FLocalizationTargetDetailCustomization::ImportAllCultures));
-			ToolBarBuilder.AddToolBarButton(FLocalizationTargetEditorCommands::Get().ImportAllCultures, NAME_None, TAttribute<FText>(), TAttribute<FText>(), FSlateIcon(FEditorStyle::GetStyleSetName(), "LocalizationTargetEditor.ImportForAllCultures"));
+		CommandList->MapAction( FLocalizationTargetEditorCommands::Get().Gather, FExecuteAction::CreateSP(this, &FLocalizationTargetDetailCustomization::Gather));
+		ToolBarBuilder.AddToolBarButton(FLocalizationTargetEditorCommands::Get().Gather, NAME_None, TAttribute<FText>(), TAttribute<FText>(), FSlateIcon(FEditorStyle::GetStyleSetName(), "LocalizationTargetEditor.Gather"));
 
-			CommandList->MapAction( FLocalizationTargetEditorCommands::Get().ExportAllCultures, FExecuteAction::CreateSP(this, &FLocalizationTargetDetailCustomization::ExportAllCultures));
-			ToolBarBuilder.AddToolBarButton(FLocalizationTargetEditorCommands::Get().ExportAllCultures, NAME_None, TAttribute<FText>(), TAttribute<FText>(), FSlateIcon(FEditorStyle::GetStyleSetName(), "LocalizationTargetEditor.ExportForAllCultures"));
+		CommandList->MapAction( FLocalizationTargetEditorCommands::Get().ImportAllCultures, FExecuteAction::CreateSP(this, &FLocalizationTargetDetailCustomization::ImportAllCultures));
+		ToolBarBuilder.AddToolBarButton(FLocalizationTargetEditorCommands::Get().ImportAllCultures, NAME_None, TAttribute<FText>(), TAttribute<FText>(), FSlateIcon(FEditorStyle::GetStyleSetName(), "LocalizationTargetEditor.ImportForAllCultures"));
 
-			CommandList->MapAction( FLocalizationTargetEditorCommands::Get().RefreshWordCounts, FExecuteAction::CreateSP(this, &FLocalizationTargetDetailCustomization::RefreshWordCounts));
-			ToolBarBuilder.AddToolBarButton(FLocalizationTargetEditorCommands::Get().RefreshWordCounts, NAME_None, TAttribute<FText>(), TAttribute<FText>(), FSlateIcon(FEditorStyle::GetStyleSetName(), "LocalizationTargetEditor.RefreshWordCounts"));
+		CommandList->MapAction( FLocalizationTargetEditorCommands::Get().ExportAllCultures, FExecuteAction::CreateSP(this, &FLocalizationTargetDetailCustomization::ExportAllCultures));
+		ToolBarBuilder.AddToolBarButton(FLocalizationTargetEditorCommands::Get().ExportAllCultures, NAME_None, TAttribute<FText>(), TAttribute<FText>(), FSlateIcon(FEditorStyle::GetStyleSetName(), "LocalizationTargetEditor.ExportForAllCultures"));
 
-			CommandList->MapAction( FLocalizationTargetEditorCommands::Get().Compile, FExecuteAction::CreateSP(this, &FLocalizationTargetDetailCustomization::Compile));
-			ToolBarBuilder.AddToolBarButton(FLocalizationTargetEditorCommands::Get().Compile, NAME_None, TAttribute<FText>(), TAttribute<FText>(), FSlateIcon(FEditorStyle::GetStyleSetName(), "LocalizationTargetEditor.Compile"));
+		CommandList->MapAction( FLocalizationTargetEditorCommands::Get().RefreshWordCounts, FExecuteAction::CreateSP(this, &FLocalizationTargetDetailCustomization::RefreshWordCounts));
+		ToolBarBuilder.AddToolBarButton(FLocalizationTargetEditorCommands::Get().RefreshWordCounts, NAME_None, TAttribute<FText>(), TAttribute<FText>(), FSlateIcon(FEditorStyle::GetStyleSetName(), "LocalizationTargetEditor.RefreshWordCounts"));
 
+		CommandList->MapAction( FLocalizationTargetEditorCommands::Get().Compile, FExecuteAction::CreateSP(this, &FLocalizationTargetDetailCustomization::Compile));
+		ToolBarBuilder.AddToolBarButton(FLocalizationTargetEditorCommands::Get().Compile, NAME_None, TAttribute<FText>(), TAttribute<FText>(), FSlateIcon(FEditorStyle::GetStyleSetName(), "LocalizationTargetEditor.Compile"));
 
-			BuildListedCulturesList();
+		BuildListedCulturesList();
 
-			DetailCategoryBuilder.AddCustomRow( LOCTEXT("SupportedCultureNamesFilterString", "Supported Culture Names") )
-				.WholeRowContent()
+		DetailCategoryBuilder.AddCustomRow( SupportedCulturesStatisticsPropertyHandle->GetPropertyDisplayName() )
+			.WholeRowContent()
+			[
+				SNew(SVerticalBox)
+				+SVerticalBox::Slot()
+				.AutoHeight()
 				[
-					SNew(SVerticalBox)
-					+SVerticalBox::Slot()
+					ToolBarBuilder.MakeWidget()
+				]
+				+SVerticalBox::Slot()
 					.AutoHeight()
 					[
-						ToolBarBuilder.MakeWidget()
+						SAssignNew(SupportedCultureListView, SListView< TSharedPtr<IPropertyHandle> >)
+						.OnGenerateRow(this, &FLocalizationTargetDetailCustomization::OnGenerateCultureRow)
+						.ListItemsSource(&ListedCultureStatisticProperties)
+						.SelectionMode(ESelectionMode::None)
+						.HeaderRow
+						(
+						SNew(SHeaderRow)
+						+SHeaderRow::Column("Culture")
+						.DefaultLabel( NSLOCTEXT("LocalizationCulture", "CultureColumnLabel", "Culture"))
+						.HAlignHeader(HAlign_Fill)
+						.HAlignCell(HAlign_Fill)
+						.VAlignCell(VAlign_Center)
+						+SHeaderRow::Column("WordCount")
+						.DefaultLabel( NSLOCTEXT("LocalizationCulture", "WordCountColumnLabel", "Word Count"))
+						.HAlignHeader(HAlign_Center)
+						.HAlignCell(HAlign_Fill)
+						.VAlignCell(VAlign_Center)
+						+SHeaderRow::Column("Actions")
+						.DefaultLabel( NSLOCTEXT("LocalizationCulture", "ActionsColumnLabel", "Actions"))
+						.HAlignHeader(HAlign_Center)
+						.HAlignCell(HAlign_Center)
+						.VAlignCell(VAlign_Center)
+						)
 					]
-					+SVerticalBox::Slot()
-						.AutoHeight()
+				+SVerticalBox::Slot()
+					.AutoHeight()
+					.VAlign(VAlign_Center)
+					[
+						SAssignNew(AddNewSupportedCultureComboButton, SComboButton)
+						.ButtonContent()
 						[
-							SAssignNew(SupportedCultureListView, SListView< TSharedPtr<IPropertyHandle> >)
-							.OnGenerateRow(this, &FLocalizationTargetDetailCustomization::OnGenerateCultureRow)
-							.ListItemsSource(&ListedCultureStatisticProperties)
-							.SelectionMode(ESelectionMode::None)
-							.HeaderRow
-							(
-							SNew(SHeaderRow)
-							+SHeaderRow::Column("Culture")
-							.DefaultLabel( NSLOCTEXT("LocalizationCulture", "CultureColumnLabel", "Culture"))
-							.HAlignHeader(HAlign_Fill)
-							.HAlignCell(HAlign_Fill)
-							.VAlignCell(VAlign_Center)
-							+SHeaderRow::Column("WordCount")
-							.DefaultLabel( NSLOCTEXT("LocalizationCulture", "WordCountColumnLabel", "Word Count"))
-							.HAlignHeader(HAlign_Center)
-							.HAlignCell(HAlign_Fill)
-							.VAlignCell(VAlign_Center)
-							+SHeaderRow::Column("Actions")
-							.DefaultLabel( NSLOCTEXT("LocalizationCulture", "ActionsColumnLabel", "Actions"))
-							.HAlignHeader(HAlign_Center)
-							.HAlignCell(HAlign_Center)
-							.VAlignCell(VAlign_Center)
-							)
+							SNew(STextBlock)
+							.Text(NSLOCTEXT("LocalizationCulture", "AddNewCultureButtonLabel", "Add New Culture"))
 						]
-					+SVerticalBox::Slot()
-						.AutoHeight()
-						.VAlign(VAlign_Center)
-						[
-							SAssignNew(AddNewSupportedCultureComboButton, SComboButton)
-							.ButtonContent()
+						.MenuContent()
 							[
-								SNew(STextBlock)
-								.Text(NSLOCTEXT("LocalizationCulture", "AddNewCultureButtonLabel", "Add New Culture"))
-							]
-							.MenuContent()
+								SNew(SBox)
+								.MaxDesiredHeight(400.0f)
+								.MaxDesiredWidth(300.0f)
 								[
-									SNew(SBox)
-									.MaxDesiredHeight(400.0f)
-									.MaxDesiredWidth(300.0f)
-									[
-										SAssignNew(SupportedCulturePicker, SCulturePicker)
-										.OnSelectionChanged(this, &FLocalizationTargetDetailCustomization::OnNewSupportedCultureSelected)
-										.IsCulturePickable(this, &FLocalizationTargetDetailCustomization::IsCultureSelectableAsSupported)
-									]
+									SAssignNew(SupportedCulturePicker, SCulturePicker)
+									.OnSelectionChanged(this, &FLocalizationTargetDetailCustomization::OnNewSupportedCultureSelected)
+									.IsCulturePickable(this, &FLocalizationTargetDetailCustomization::IsCultureSelectableAsSupported)
 								]
-						]
-				];
+							]
+					]
+			];
+	});
+
+	UStructProperty* const SettingsStructProperty = CastChecked<UStructProperty>(TargetSettingsPropertyHandle->GetProperty());
+	for (TFieldIterator<UProperty> Iterator(SettingsStructProperty->Struct); Iterator; ++Iterator)
+	{
+		UProperty* const MemberProperty = *Iterator;
+		
+		const bool IsEditable = MemberProperty->HasAnyPropertyFlags(CPF_Edit);
+		if (IsEditable)
+		{
+			const FName CategoryName = FObjectEditorUtils::GetCategoryFName(MemberProperty);
+			IDetailCategoryBuilder& DetailCategoryBuilder = DetailBuilder.EditCategory(CategoryName);
+
+			const FName PropertyName = MemberProperty->GetFName();
+
+			const TSharedPtr<IPropertyHandle> MemberPropertyHandle = TargetSettingsPropertyHandle->GetChildHandle(MemberProperty->GetFName());
+			if (MemberPropertyHandle.IsValid() && MemberPropertyHandle->IsValidHandle())
+			{
+				FPropertyCustomizationFunction* const Function = PropertyCustomizationMap.Find(PropertyName);
+				if (Function)
+				{
+					MemberPropertyHandle->MarkHiddenByCustomization();
+					(*Function)(MemberPropertyHandle.ToSharedRef(), DetailCategoryBuilder);
+				}
+				else
+				{
+					DetailCategoryBuilder.AddProperty(MemberPropertyHandle);
+				}
+			}
 		}
 	}
+
+	//{
+	//	IDetailCategoryBuilder& ServiceProviderCategoryBuilder = DetailBuilder.EditCategory(FName("LocalizationServiceProvider"));
+
+	//	ILocalizationServiceProvider* const LSP = ILocalizationDashboardModule::Get().GetCurrentLocalizationServiceProvider();
+	//	if (LSP && LocalizationTarget)
+	//	{
+	//		LSP->CustomizeTargetDetails(ServiceProviderCategoryBuilder, *LocalizationTarget);
+	//	}
+	//}
 }
 
 FLocalizationTargetSettings* FLocalizationTargetDetailCustomization::GetTargetSettings() const
diff --git a/Engine/Source/Editor/LocalizationDashboard/Private/LocalizationTargetSetDetailCustomization.cpp b/Engine/Source/Editor/LocalizationDashboard/Private/LocalizationTargetSetDetailCustomization.cpp
--- a/Engine/Source/Editor/LocalizationDashboard/Private/LocalizationTargetSetDetailCustomization.cpp
+++ b/Engine/Source/Editor/LocalizationDashboard/Private/LocalizationTargetSetDetailCustomization.cpp
@@ -9,6 +9,7 @@
 #include "DetailCategoryBuilder.h"
 #include "IDetailsView.h"
 #include "LocalizationCommandletTasks.h"
+#include "ObjectEditorUtils.h"
 
 #define LOCTEXT_NAMESPACE "LocalizationDashboard"
 
@@ -80,10 +81,12 @@ void FLocalizationTargetSetDetailCustomization::CustomizeDetails(IDetailLayoutBu
 	}
 
 	{
-		IDetailCategoryBuilder& DetailCategoryBuilder = DetailBuilder.EditCategory(FName("Targets"));
 		TargetObjectsPropertyHandle = DetailBuilder.GetProperty(GET_MEMBER_NAME_CHECKED(ULocalizationTargetSet,TargetObjects));
 		if (TargetObjectsPropertyHandle->IsValidHandle())
 		{
+			const FName CategoryName = FObjectEditorUtils::GetCategoryFName(TargetObjectsPropertyHandle->GetProperty());
+			IDetailCategoryBuilder& DetailCategoryBuilder = DetailBuilder.EditCategory(CategoryName);
+
 			TargetObjectsPropertyHandle->MarkHiddenByCustomization();
 			TargetsArrayPropertyHandle_OnNumElementsChanged = FSimpleDelegate::CreateSP(this, &FLocalizationTargetSetDetailCustomization::RebuildTargetsList);
 			TargetObjectsPropertyHandle->AsArray()->SetOnNumElementsChanged(TargetsArrayPropertyHandle_OnNumElementsChanged);
@@ -106,7 +109,7 @@ void FLocalizationTargetSetDetailCustomization::CustomizeDetails(IDetailLayoutBu
 
 			BuildTargetsList();
 
-			DetailCategoryBuilder.AddCustomRow( LOCTEXT("TargetsFilterString", "Targets") )
+			DetailCategoryBuilder.AddCustomRow(TargetObjectsPropertyHandle->GetPropertyDisplayName())
 				.WholeRowContent()
 				[
 					SNew(SVerticalBox)
