Gathering text for localization via the localization dashboard/target editor now uses gathering configuration that is exposed via the UI.

diff --git a/Engine/Source/Editor/LocalizationDashboard/Classes/LocalizationTarget.h b/Engine/Source/Editor/LocalizationDashboard/Classes/LocalizationTarget.h
--- a/Engine/Source/Editor/LocalizationDashboard/Classes/LocalizationTarget.h
+++ b/Engine/Source/Editor/LocalizationDashboard/Classes/LocalizationTarget.h
@@ -13,6 +13,66 @@ enum class ELocalizationTargetStatus : uint8
 	Clear
 };
 
+USTRUCT()
+struct FGatherTextFromTextFilesConfiguration
+{
+	GENERATED_USTRUCT_BODY()
+
+	UPROPERTY(config, EditAnywhere, Category = "Filter")
+	TArray<FString> SearchDirectories;
+
+	UPROPERTY(config, EditAnywhere, Category = "Filter")
+	TArray<FString> ExcludePathWildcards;
+
+	UPROPERTY(config, EditAnywhere, Category = "Filter")
+	TArray<FString> FileExtensions;
+};
+
+USTRUCT()
+struct FGatherTextFromPackagesConfiguration
+{
+	GENERATED_USTRUCT_BODY()
+
+	UPROPERTY(config, EditAnywhere, Category = "Filter")
+	TArray<FString> IncludePathWildcards;
+
+	UPROPERTY(config, EditAnywhere, Category = "Filter")
+	TArray<FString> ExcludePathWildcards;
+
+	UPROPERTY(config, EditAnywhere, Category = "Filter")
+	TArray<FString> FileExtensions;
+};
+
+USTRUCT()
+struct FMetaDataKeyGatherSpecification
+{
+	GENERATED_USTRUCT_BODY()
+
+	UPROPERTY(config, EditAnywhere, Category = "Input")
+	FString MetaDataKey;
+
+	UPROPERTY(config, EditAnywhere, Category = "Output")
+	FString TextNamespace;
+
+	UPROPERTY(config, EditAnywhere, Category = "Output")
+	FString TextKeyPattern;
+};
+
+USTRUCT()
+struct FGatherTextFromMetaDataConfiguration
+{
+	GENERATED_USTRUCT_BODY()
+
+	UPROPERTY(config, EditAnywhere, Category = "Filter")
+	TArray<FString> IncludePathWildcards;
+
+	UPROPERTY(config, EditAnywhere, Category = "Filter")
+	TArray<FString> ExcludePathWildcards;
+
+	UPROPERTY(config, EditAnywhere, Category = "MetaData")
+	TArray<FMetaDataKeyGatherSpecification> KeySpecifications;
+};
+
 USTRUCT()
 struct FCultureStatistics
 {
@@ -44,10 +104,19 @@ struct FLocalizationTargetSettings
 	UPROPERTY(Transient, EditAnywhere, Category = "Target")
 	ELocalizationTargetStatus Status;
 
-	UPROPERTY(config, EditAnywhere, Category = "Target")
+	UPROPERTY(config, EditAnywhere, Category = "Gather Configuration")
 	TArray<FString> TargetDependencies;
 
-	UPROPERTY(config, EditAnywhere, Category = "Target", AdvancedDisplay, meta=(FilePathFilter="manifest"))
+	UPROPERTY(config, EditAnywhere, Category = "Gather Configuration")
+	FGatherTextFromTextFilesConfiguration GatherFromTextFiles;
+
+	UPROPERTY(config, EditAnywhere, Category = "Gather Configuration")
+	FGatherTextFromPackagesConfiguration GatherFromPackages;
+
+	UPROPERTY(config, EditAnywhere, Category = "Gather Configuration")
+	FGatherTextFromMetaDataConfiguration GatherFromMetaData;
+
+	UPROPERTY(config, EditAnywhere, Category = "Gather Configuration", AdvancedDisplay, meta=(FilePathFilter="manifest"))
 	TArray<FFilePath> AdditionalManifestDependencies;
 
 	UPROPERTY(config, EditAnywhere, Category = "Cultures")
diff --git a/Engine/Source/Editor/LocalizationDashboard/Private/LocalizationConfigurationScript.cpp b/Engine/Source/Editor/LocalizationDashboard/Private/LocalizationConfigurationScript.cpp
--- a/Engine/Source/Editor/LocalizationDashboard/Private/LocalizationConfigurationScript.cpp
+++ b/Engine/Source/Editor/LocalizationDashboard/Private/LocalizationConfigurationScript.cpp
@@ -115,6 +115,7 @@ namespace LocalizationConfigurationScript
 
 		uint32 GatherTextStepIndex = 0;
 		// GatherTextFromSource
+		if (Target.GatherFromTextFiles.SearchDirectories.Num() && Target.GatherFromTextFiles.FileExtensions.Num()) // Don't gather from source if there are no source files.
 		{
 			FConfigSection& ConfigSection = Script.GatherTextStep(GatherTextStepIndex++);
 
@@ -122,40 +123,80 @@ namespace LocalizationConfigurationScript
 			ConfigSection.Add( TEXT("CommandletClass"), TEXT("GatherTextFromSource") );
 
 			// Include Paths
-			if (HasSourceCode) 
+			for (const auto& IncludePath : Target.GatherFromTextFiles.SearchDirectories)
 			{
-				ConfigSection.Add( TEXT("IncludePaths"), SourceDirRelativeToGameDir );
+				ConfigSection.Add( TEXT("IncludePaths"), IncludePath );
 			}
-			ConfigSection.Add( TEXT("IncludePaths"), ConfigDirRelativeToGameDir );
 
 			// Exclude Paths
 			ConfigSection.Add( TEXT("ExcludePaths"), ConfigDirRelativeToGameDir / TEXT("Localization") );
+			for (const auto& ExcludePath : Target.GatherFromTextFiles.ExcludePathWildcards)
+			{
+				ConfigSection.Add( TEXT("ExcludePaths"), ExcludePath );
+			}
 
-			// SourceFileSearchFilters
-			if (HasSourceCode)
+			// Source File Search Filters
+			for (const auto& FileExtension : Target.GatherFromTextFiles.FileExtensions)
 			{
-				ConfigSection.Add( TEXT("SourceFileSearchFilters"), TEXT("*.h") );
-				ConfigSection.Add( TEXT("SourceFileSearchFilters"), TEXT("*.cpp") );
+				ConfigSection.Add( TEXT("SourceFileSearchFilters"), FileExtension );
 			}
-			ConfigSection.Add( TEXT("SourceFileSearchFilters"), TEXT("*.ini") );
 		}
 
 		// GatherTextFromAssets
+		if (Target.GatherFromTextFiles.SearchDirectories.Num() && Target.GatherFromTextFiles.FileExtensions.Num()) // Don't gather from packages if there are none specified.
 		{
 			FConfigSection& ConfigSection = Script.GatherTextStep(GatherTextStepIndex++);
 
 			// CommandletClass
 			ConfigSection.Add( TEXT("CommandletClass"), TEXT("GatherTextFromAssets") );
 
-			// IncludePaths
-			ConfigSection.Add( TEXT("IncludePaths"), ContentDirRelativeToGameDir );
+			// Include Paths
+			for (const auto& IncludePath : Target.GatherFromPackages.IncludePathWildcards)
+			{
+				ConfigSection.Add( TEXT("IncludePaths"), IncludePath );
+			}
 
-			// ExcludePaths
+			// Exclude Paths
 			ConfigSection.Add( TEXT("ExcludePaths"), ContentDirRelativeToGameDir / TEXT("Localization") );
+			for (const auto& ExcludePath : Target.GatherFromPackages.ExcludePathWildcards)
+			{
+				ConfigSection.Add( TEXT("ExcludePaths"), ExcludePath );
+			}
 
-			// PackageExtensions
-			ConfigSection.Add( TEXT("PackageExtensions"), TEXT("*.umap") );
-			ConfigSection.Add( TEXT("PackageExtensions"), TEXT("*.uasset") );
+			// Package Extensions
+			for (const auto& FileExtension : Target.GatherFromPackages.FileExtensions)
+			{
+				ConfigSection.Add( TEXT("PackageExtensions"), FileExtension );
+			}
+		}
+
+		// GatherTextFromMetadata
+		if (Target.GatherFromMetaData.IncludePathWildcards.Num() && Target.GatherFromMetaData.MetaData.Num()) // Don't gather from metadata if there are none specified.
+		{
+			FConfigSection& ConfigSection = Script.GatherTextStep(GatherTextStepIndex++);
+
+			// CommandletClass
+			ConfigSection.Add( TEXT("CommandletClass"), TEXT("GatherTextFromMetadata") );
+
+			// Include Paths
+			for (const auto& IncludePath : Target.GatherFromMetaData.IncludePathWildcards)
+			{
+				ConfigSection.Add( TEXT("IncludePaths"), IncludePath );
+			}
+
+			// Exclude Paths
+			for (const auto& ExcludePath : Target.GatherFromMetaData.ExcludePathWildcards)
+			{
+				ConfigSection.Add( TEXT("ExcludePaths"), ExcludePath );
+			}
+
+			// Package Extensions
+			for (const FMetaDataTextGenerationSpecification& Specification : Target.GatherFromMetaData.MetaData)
+			{
+				ConfigSection.Add( TEXT("InputKeys"), Specification.MetadataKey );
+				ConfigSection.Add( TEXT("OutputNamespaces"), Specification.TextNamespace );
+				ConfigSection.Add( TEXT("OutputKeys"), FString::Printf(TEXT("\"%s\""), *Specification.TextKeyPattern) );
+			}
 		}
 
 		// GenerateGatherManifest
@@ -167,7 +208,6 @@ namespace LocalizationConfigurationScript
 		}
 
 		// GenerateGatherArchive
-		if (Target.SupportedCulturesStatistics.Num())
 		{
 			FConfigSection& ConfigSection = Script.GatherTextStep(GatherTextStepIndex++);
 
