Submitting previously unsaved changes after refactoring the names of struct/class member variables in LocalizationTarget.h.

diff --git a/Engine/Source/Editor/LocalizationDashboard/Private/LocalizationConfigurationScript.cpp b/Engine/Source/Editor/LocalizationDashboard/Private/LocalizationConfigurationScript.cpp
--- a/Engine/Source/Editor/LocalizationDashboard/Private/LocalizationConfigurationScript.cpp
+++ b/Engine/Source/Editor/LocalizationDashboard/Private/LocalizationConfigurationScript.cpp
@@ -171,7 +171,7 @@ namespace LocalizationConfigurationScript
 		}
 
 		// GatherTextFromMetadata
-		if (Target.GatherFromMetaData.IncludePathWildcards.Num() && Target.GatherFromMetaData.MetaData.Num()) // Don't gather from metadata if there are none specified.
+		if (Target.GatherFromMetaData.IncludePathWildcards.Num() && Target.GatherFromMetaData.KeySpecifications.Num()) // Don't gather from metadata if there are none specified.
 		{
 			FConfigSection& ConfigSection = Script.GatherTextStep(GatherTextStepIndex++);
 
@@ -191,9 +191,9 @@ namespace LocalizationConfigurationScript
 			}
 
 			// Package Extensions
-			for (const FMetaDataTextGenerationSpecification& Specification : Target.GatherFromMetaData.MetaData)
+			for (const FMetaDataKeyGatherSpecification& Specification : Target.GatherFromMetaData.KeySpecifications)
 			{
-				ConfigSection.Add( TEXT("InputKeys"), Specification.MetadataKey );
+				ConfigSection.Add( TEXT("InputKeys"), Specification.MetaDataKey );
 				ConfigSection.Add( TEXT("OutputNamespaces"), Specification.TextNamespace );
 				ConfigSection.Add( TEXT("OutputKeys"), FString::Printf(TEXT("\"%s\""), *Specification.TextKeyPattern) );
 			}
