Missing include for non unity build in Localization Commandlet code.

diff --git a/Engine/Source/Editor/LocalizationDashboard/Private/SLocalizationTargetStatusButton.cpp b/Engine/Source/Editor/LocalizationDashboard/Private/SLocalizationTargetStatusButton.cpp
--- a/Engine/Source/Editor/LocalizationDashboard/Private/SLocalizationTargetStatusButton.cpp
+++ b/Engine/Source/Editor/LocalizationDashboard/Private/SLocalizationTargetStatusButton.cpp
@@ -3,6 +3,7 @@
 #include "LocalizationDashboardPrivatePCH.h"
 #include "SLocalizationTargetStatusButton.h"
 #include "LocalizationCommandletExecution.h"
+#include "LocalizationCommandletTasks.h"
 
 #define LOCTEXT_NAMESPACE "LocalizationTargetStatusButton"
 
