Relocated several methods on FLocalizationTargetSettings to ULocalizationTarget in preparation for seprating UE4 and game-specific target sets. ULocalizationTarget will make the distinction possible, not FLocalizationTargetSettings.

diff --git a/Engine/Source/Editor/LocalizationDashboard/Classes/LocalizationTarget.h b/Engine/Source/Editor/LocalizationDashboard/Classes/LocalizationTarget.h
--- a/Engine/Source/Editor/LocalizationDashboard/Classes/LocalizationTarget.h
+++ b/Engine/Source/Editor/LocalizationDashboard/Classes/LocalizationTarget.h
@@ -124,11 +124,6 @@ struct FLocalizationTargetSettings
 
 	UPROPERTY(config, EditAnywhere, Category = "Cultures")
 	TArray<FCultureStatistics> SupportedCulturesStatistics;
-
-	bool UpdateWordCountsFromCSV();
-	void UpdateStatusFromConflictReport();
-	bool Rename(const FString& NewName);
-	bool DeleteFiles(const FString* const Culture = nullptr) const;
 };
 
 UCLASS(Within=LocalizationTargetSet)
@@ -144,6 +139,11 @@ public:
 #if WITH_EDITOR
 	virtual void PostEditChangeProperty(FPropertyChangedEvent& PropertyChangedEvent) override;
 #endif
+
+	bool UpdateWordCountsFromCSV();
+	void UpdateStatusFromConflictReport();
+	bool RenameTargetAndFiles(const FString& NewName);
+	bool DeleteFiles(const FString* const Culture = nullptr) const;
 };
 
 UCLASS(Config=Game, defaultconfig)
diff --git a/Engine/Source/Editor/LocalizationDashboard/Private/LocalizationDashboardModule.cpp b/Engine/Source/Editor/LocalizationDashboard/Private/LocalizationDashboardModule.cpp
--- a/Engine/Source/Editor/LocalizationDashboard/Private/LocalizationDashboardModule.cpp
+++ b/Engine/Source/Editor/LocalizationDashboard/Private/LocalizationDashboardModule.cpp
@@ -22,8 +22,8 @@ public:
 		
 		for (ULocalizationTarget*& Target : Settings->TargetObjects)
 		{
-			Target->Settings.UpdateStatusFromConflictReport();
-			Target->Settings.UpdateWordCountsFromCSV();
+			Target->UpdateStatusFromConflictReport();
+			Target->UpdateWordCountsFromCSV();
 		}
 	}
 
diff --git a/Engine/Source/Editor/LocalizationDashboard/Private/LocalizationTarget.cpp b/Engine/Source/Editor/LocalizationDashboard/Private/LocalizationTarget.cpp
--- a/Engine/Source/Editor/LocalizationDashboard/Private/LocalizationTarget.cpp
+++ b/Engine/Source/Editor/LocalizationDashboard/Private/LocalizationTarget.cpp
@@ -155,35 +155,48 @@ namespace
 	};
 }
 
-bool FLocalizationTargetSettings::UpdateWordCountsFromCSV()
+void ULocalizationTarget::PostEditChangeProperty(FPropertyChangedEvent& PropertyChangedEvent)
+{
+	ULocalizationTargetSet* const LocalizationTargetSet = Cast<ULocalizationTargetSet>(GetOuter());
+	if (LocalizationTargetSet)
+	{
+		UClass* ParentClass = LocalizationTargetSet ? LocalizationTargetSet->GetClass() : nullptr;
+		UProperty* ChangedParentProperty = ParentClass ? ParentClass->FindPropertyByName(GET_MEMBER_NAME_CHECKED(ULocalizationTargetSet,TargetObjects)) : nullptr;
+
+		FPropertyChangedEvent PropertyChangedEventForParent(ChangedParentProperty);
+		LocalizationTargetSet->PostEditChangeProperty(PropertyChangedEventForParent);
+	}
+}
+
+bool ULocalizationTarget::UpdateWordCountsFromCSV()
 {
-	const FString CSVFilePath = LocalizationConfigurationScript::GetWordCountCSVPath(*this);
-	return FWordCountCSVParser::Execute(*this, CSVFilePath);
+	const FString CSVFilePath = LocalizationConfigurationScript::GetWordCountCSVPath(Settings);
+	return FWordCountCSVParser::Execute(Settings, CSVFilePath);
 }
 
-void FLocalizationTargetSettings::UpdateStatusFromConflictReport()
+void ULocalizationTarget::UpdateStatusFromConflictReport()
 {
-	const FString ConflictReportFilePath = LocalizationConfigurationScript::GetConflictReportPath(*this);
+	const FString ConflictReportFilePath = LocalizationConfigurationScript::GetConflictReportPath(Settings);
 	const int64 ConflictReportFileSize = IFileManager::Get().FileSize(*ConflictReportFilePath);
 
 	switch (ConflictReportFileSize)
 	{
 	case INDEX_NONE:
 		// Unknown
-		Status = ELocalizationTargetStatus::Unknown;
+		Settings.Status = ELocalizationTargetStatus::Unknown;
 		break;
 	case 0:
 		// Good
-		Status = ELocalizationTargetStatus::Clear;
+		Settings.Status = ELocalizationTargetStatus::Clear;
 		break;
 	default:
 		// Bad
-		Status = ELocalizationTargetStatus::ConflictsPresent;
+		Settings.Status = ELocalizationTargetStatus::ConflictsPresent;
 		break;
 	}
 }
 
-bool FLocalizationTargetSettings::Rename(const FString& NewName)
+bool ULocalizationTarget::RenameTargetAndFiles(const FString& NewName)
 {
 	bool HasCompletelySucceeded = true;
 	ISourceControlModule& SourceControl = ISourceControlModule::Get();
@@ -230,7 +243,7 @@ bool FLocalizationTargetSettings::Rename(const FString& NewName)
 	};
 
 	// Delete configuration files.
-	const TArray<FString>& ScriptPaths = LocalizationConfigurationScript::GetScriptPaths(*this);
+	const TArray<FString>& ScriptPaths = LocalizationConfigurationScript::GetScriptPaths(Settings);
 	for (const FString& ScriptPath : ScriptPaths)
 	{
 		if (!TryDelete(ScriptPath))
@@ -241,32 +254,32 @@ bool FLocalizationTargetSettings::Rename(const FString& NewName)
 
 	const auto& GetNamedPaths = [this](TArray<FString>& NamedPaths)
 	{
-		NamedPaths.Add(LocalizationConfigurationScript::GetManifestPath(*this));
-		NamedPaths.Add(LocalizationConfigurationScript::GetWordCountCSVPath(*this));
-		NamedPaths.Add(LocalizationConfigurationScript::GetConflictReportPath(*this));
+		NamedPaths.Add(LocalizationConfigurationScript::GetManifestPath(Settings));
+		NamedPaths.Add(LocalizationConfigurationScript::GetWordCountCSVPath(Settings));
+		NamedPaths.Add(LocalizationConfigurationScript::GetConflictReportPath(Settings));
 
 		TArray<const FCultureStatistics*> Cultures;
-		Cultures.Add(&NativeCultureStatistics);
-		for (const FCultureStatistics& Culture : SupportedCulturesStatistics)
+		Cultures.Add(&Settings.NativeCultureStatistics);
+		for (const FCultureStatistics& Culture : Settings.SupportedCulturesStatistics)
 		{
 			Cultures.Add(&Culture);
 		}
 
 		for (const FCultureStatistics* Culture : Cultures)
 		{
-			NamedPaths.Add(LocalizationConfigurationScript::GetArchivePath(*this, Culture->CultureName));
-			NamedPaths.Add(LocalizationConfigurationScript::GetDefaultPOPath(*this, Culture->CultureName));
-			NamedPaths.Add(LocalizationConfigurationScript::GetLocResPath(*this, Culture->CultureName));
+			NamedPaths.Add(LocalizationConfigurationScript::GetArchivePath(Settings, Culture->CultureName));
+			NamedPaths.Add(LocalizationConfigurationScript::GetDefaultPOPath(Settings, Culture->CultureName));
+			NamedPaths.Add(LocalizationConfigurationScript::GetLocResPath(Settings, Culture->CultureName));
 		}
 
-		NamedPaths.Add(LocalizationConfigurationScript::GetDataDirectory(*this));
+		NamedPaths.Add(LocalizationConfigurationScript::GetDataDirectory(Settings));
 	};
 
 	TArray<FString> OldPaths;
 	GetNamedPaths(OldPaths);
 
 	// Rename
-	Name = NewName;
+	Settings.Name = NewName;
 
 	TArray<FString> NewPaths;
 	GetNamedPaths(NewPaths);
@@ -320,29 +333,29 @@ bool FLocalizationTargetSettings::Rename(const FString& NewName)
 	}
 
 	// Generate new configuration files.
-	LocalizationConfigurationScript::GenerateGatherScript(*this).Write(LocalizationConfigurationScript::GetGatherScriptPath(*this));
-	LocalizationConfigurationScript::GenerateImportScript(*this).Write(LocalizationConfigurationScript::GetImportScriptPath(*this));
-	LocalizationConfigurationScript::GenerateExportScript(*this).Write(LocalizationConfigurationScript::GetExportScriptPath(*this));
-	LocalizationConfigurationScript::GenerateReportScript(*this).Write(LocalizationConfigurationScript::GetReportScriptPath(*this));
+	LocalizationConfigurationScript::GenerateGatherScript(Settings).Write(LocalizationConfigurationScript::GetGatherScriptPath(Settings));
+	LocalizationConfigurationScript::GenerateImportScript(Settings).Write(LocalizationConfigurationScript::GetImportScriptPath(Settings));
+	LocalizationConfigurationScript::GenerateExportScript(Settings).Write(LocalizationConfigurationScript::GetExportScriptPath(Settings));
+	LocalizationConfigurationScript::GenerateReportScript(Settings).Write(LocalizationConfigurationScript::GetReportScriptPath(Settings));
 
 	//if (CanUseSourceControl)
 	//{
-	//	HasCompletelySucceeded = HasCompletelySucceeded && SourceControlProvider.Execute(ISourceControlOperation::Create<FMarkForAdd>(), *LocalizationConfigurationScript::GetGatherScriptPath(*this)) == ECommandResult::Succeeded;
-	//	HasCompletelySucceeded = HasCompletelySucceeded && SourceControlProvider.Execute(ISourceControlOperation::Create<FMarkForAdd>(), *LocalizationConfigurationScript::GetImportScriptPath(*this)) == ECommandResult::Succeeded;
-	//	HasCompletelySucceeded = HasCompletelySucceeded && SourceControlProvider.Execute(ISourceControlOperation::Create<FMarkForAdd>(), *LocalizationConfigurationScript::GetExportScriptPath(*this)) == ECommandResult::Succeeded;
-	//	HasCompletelySucceeded = HasCompletelySucceeded && SourceControlProvider.Execute(ISourceControlOperation::Create<FMarkForAdd>(), *LocalizationConfigurationScript::GetReportScriptPath(*this)) == ECommandResult::Succeeded;
+	//	HasCompletelySucceeded = HasCompletelySucceeded && SourceControlProvider.Execute(ISourceControlOperation::Create<FMarkForAdd>(), *LocalizationConfigurationScript::GetGatherScriptPath(Settings)) == ECommandResult::Succeeded;
+	//	HasCompletelySucceeded = HasCompletelySucceeded && SourceControlProvider.Execute(ISourceControlOperation::Create<FMarkForAdd>(), *LocalizationConfigurationScript::GetImportScriptPath(Settings)) == ECommandResult::Succeeded;
+	//	HasCompletelySucceeded = HasCompletelySucceeded && SourceControlProvider.Execute(ISourceControlOperation::Create<FMarkForAdd>(), *LocalizationConfigurationScript::GetExportScriptPath(Settings)) == ECommandResult::Succeeded;
+	//	HasCompletelySucceeded = HasCompletelySucceeded && SourceControlProvider.Execute(ISourceControlOperation::Create<FMarkForAdd>(), *LocalizationConfigurationScript::GetReportScriptPath(Settings)) == ECommandResult::Succeeded;
 	//}
 
 	return HasCompletelySucceeded;
 }
 
-bool FLocalizationTargetSettings::DeleteFiles(const FString* const Culture) const
+bool ULocalizationTarget::DeleteFiles(const FString* const Culture) const
 {
 	bool HasCompletelySucceeded = true;
 	ISourceControlModule& SourceControl = ISourceControlModule::Get();
 	ISourceControlProvider& SourceControlProvider = SourceControl.GetProvider();
 
-	const FString DataDirectory = LocalizationConfigurationScript::GetDataDirectory(*this) + (Culture ? *Culture : FString());
+	const FString DataDirectory = LocalizationConfigurationScript::GetDataDirectory(Settings) + (Culture ? *Culture : FString());
 
 	// Remove data files from source control.
 	{
@@ -386,7 +399,7 @@ bool FLocalizationTargetSettings::DeleteFiles(const FString* const Culture) cons
 	// Delete configuration files if deleting the target.
 	if (!Culture)
 	{
-		const TArray<FString>& ScriptPaths = LocalizationConfigurationScript::GetScriptPaths(*this);
+		const TArray<FString>& ScriptPaths = LocalizationConfigurationScript::GetScriptPaths(Settings);
 
 		// Remove script files from source control.
 		{
@@ -431,19 +444,6 @@ bool FLocalizationTargetSettings::DeleteFiles(const FString* const Culture) cons
 	return HasCompletelySucceeded;
 }
 
-void ULocalizationTarget::PostEditChangeProperty(FPropertyChangedEvent& PropertyChangedEvent)
-{
-	ULocalizationTargetSet* const LocalizationTargetSet = Cast<ULocalizationTargetSet>(GetOuter());
-	if (LocalizationTargetSet)
-	{
-		UClass* ParentClass = LocalizationTargetSet ? LocalizationTargetSet->GetClass() : nullptr;
-		UProperty* ChangedParentProperty = ParentClass ? ParentClass->FindPropertyByName(GET_MEMBER_NAME_CHECKED(ULocalizationTargetSet,TargetObjects)) : nullptr;
-
-		FPropertyChangedEvent PropertyChangedEventForParent(ChangedParentProperty);
-		LocalizationTargetSet->PostEditChangeProperty(PropertyChangedEventForParent);
-	}
-}
-
 void ULocalizationTargetSet::PostInitProperties()
 {
 	Super::PostInitProperties();
diff --git a/Engine/Source/Editor/LocalizationDashboard/Private/LocalizationTargetDetailCustomization.cpp b/Engine/Source/Editor/LocalizationDashboard/Private/LocalizationTargetDetailCustomization.cpp
--- a/Engine/Source/Editor/LocalizationDashboard/Private/LocalizationTargetDetailCustomization.cpp
+++ b/Engine/Source/Editor/LocalizationDashboard/Private/LocalizationTargetDetailCustomization.cpp
@@ -404,7 +404,7 @@ void FLocalizationTargetDetailCustomization::OnTargetNameCommitted(const FText&
 				TargetNamePropertyHandle->NotifyPreChange();
 			}
 
-			TargetSettings->Rename(NewText.ToString());
+			LocalizationTarget->RenameTargetAndFiles(NewText.ToString());
 
 			if (TargetNamePropertyHandle.IsValid() && TargetNamePropertyHandle->IsValidHandle())
 			{
@@ -705,8 +705,7 @@ void FLocalizationTargetDetailCustomization::RefreshWordCounts()
 
 void FLocalizationTargetDetailCustomization::UpdateTargetFromReports()
 {
-	FLocalizationTargetSettings* const TargetSettings = GetTargetSettings();
-	if (TargetSettings)
+	if (LocalizationTarget.IsValid())
 	{
 		TArray< TSharedPtr<IPropertyHandle> > WordCountPropertyHandles;
 
@@ -746,8 +745,8 @@ void FLocalizationTargetDetailCustomization::UpdateTargetFromReports()
 		{
 			WordCountPropertyHandle->NotifyPreChange();
 		}
-		TargetSettings->UpdateWordCountsFromCSV();
-		TargetSettings->UpdateStatusFromConflictReport();
+		LocalizationTarget->UpdateWordCountsFromCSV();
+		LocalizationTarget->UpdateStatusFromConflictReport();
 		for (const TSharedPtr<IPropertyHandle>& WordCountPropertyHandle : WordCountPropertyHandles)
 		{
 			WordCountPropertyHandle->NotifyPostChange();
diff --git a/Engine/Source/Editor/LocalizationDashboard/Private/LocalizationTargetSetDetailCustomization.cpp b/Engine/Source/Editor/LocalizationDashboard/Private/LocalizationTargetSetDetailCustomization.cpp
--- a/Engine/Source/Editor/LocalizationDashboard/Private/LocalizationTargetSetDetailCustomization.cpp
+++ b/Engine/Source/Editor/LocalizationDashboard/Private/LocalizationTargetSetDetailCustomization.cpp
@@ -403,8 +403,8 @@ void FLocalizationTargetSetDetailCustomization::UpdateTargetFromReports(ULocaliz
 	//{
 	//	WordCountPropertyHandle->NotifyPreChange();
 	//}
-	LocalizationTarget->Settings.UpdateWordCountsFromCSV();
-	LocalizationTarget->Settings.UpdateStatusFromConflictReport();
+	LocalizationTarget->UpdateWordCountsFromCSV();
+	LocalizationTarget->UpdateStatusFromConflictReport();
 	//for (const TSharedPtr<IPropertyHandle>& WordCountPropertyHandle : WordCountPropertyHandles)
 	//{
 	//	WordCountPropertyHandle->NotifyPostChange();
diff --git a/Engine/Source/Editor/LocalizationDashboard/Private/SLocalizationDashboardTargetRow.cpp b/Engine/Source/Editor/LocalizationDashboard/Private/SLocalizationDashboardTargetRow.cpp
--- a/Engine/Source/Editor/LocalizationDashboard/Private/SLocalizationDashboardTargetRow.cpp
+++ b/Engine/Source/Editor/LocalizationDashboard/Private/SLocalizationDashboardTargetRow.cpp
@@ -297,8 +297,8 @@ void SLocalizationDashboardTargetRow::UpdateTargetFromReports()
 		//{
 		//	WordCountPropertyHandle->NotifyPreChange();
 		//}
-		LocalizationTarget->Settings.UpdateWordCountsFromCSV();
-		LocalizationTarget->Settings.UpdateStatusFromConflictReport();
+		LocalizationTarget->UpdateWordCountsFromCSV();
+		LocalizationTarget->UpdateStatusFromConflictReport();
 		//for (const TSharedPtr<IPropertyHandle>& WordCountPropertyHandle : WordCountPropertyHandles)
 		//{
 		//	WordCountPropertyHandle->NotifyPostChange();
@@ -472,7 +472,7 @@ void SLocalizationDashboardTargetRow::Delete()
 			{
 			case EAppReturnType::Ok:
 				{
-					LocalizationTarget->Settings.DeleteFiles();
+					LocalizationTarget->DeleteFiles();
 
 					// Close target editor.
 					if (TargetEditorDockTab.IsValid())
diff --git a/Engine/Source/Editor/LocalizationDashboard/Private/SLocalizationTargetEditorCultureRow.cpp b/Engine/Source/Editor/LocalizationDashboard/Private/SLocalizationTargetEditorCultureRow.cpp
--- a/Engine/Source/Editor/LocalizationDashboard/Private/SLocalizationTargetEditorCultureRow.cpp
+++ b/Engine/Source/Editor/LocalizationDashboard/Private/SLocalizationTargetEditorCultureRow.cpp
@@ -1,4 +1,4 @@
-// Copyright 1998-2015 Epic Games, Inc. All Rights Reserved.
+﻿// Copyright 1998-2015 Epic Games, Inc. All Rights Reserved.
 
 #include "LocalizationDashboardPrivatePCH.h"
 #include "SLocalizationTargetEditorCultureRow.h"
@@ -147,6 +147,17 @@ TSharedRef<SWidget> SLocalizationTargetEditorCultureRow::GenerateWidgetForColumn
 	return Return.IsValid() ? Return.ToSharedRef() : SNullWidget::NullWidget;
 }
 
+ULocalizationTarget* SLocalizationTargetEditorCultureRow::GetTarget() const
+{
+	if (TargetSettingsPropertyHandle.IsValid() && TargetSettingsPropertyHandle->IsValidHandle())
+	{
+		TArray<UObject*> OuterObjects;
+		TargetSettingsPropertyHandle->GetOuterObjects(OuterObjects);
+		return CastChecked<ULocalizationTarget>(OuterObjects.Top());
+	}
+	return nullptr;
+}
+
 FLocalizationTargetSettings* SLocalizationTargetEditorCultureRow::GetTargetSettings() const
 {
 	if (TargetSettingsPropertyHandle.IsValid() && TargetSettingsPropertyHandle->IsValidHandle())
@@ -237,8 +248,8 @@ TOptional<float> SLocalizationTargetEditorCultureRow::GetProgressPercentage() co
 
 void SLocalizationTargetEditorCultureRow::UpdateTargetFromReports()
 {
-	FLocalizationTargetSettings* const TargetSettings = GetTargetSettings();
-	if (TargetSettings)
+	ULocalizationTarget* const LocalizationTarget = GetTarget();
+	if (LocalizationTarget)
 	{
 		TArray< TSharedPtr<IPropertyHandle> > WordCountPropertyHandles;
 
@@ -278,8 +289,8 @@ void SLocalizationTargetEditorCultureRow::UpdateTargetFromReports()
 		{
 			WordCountPropertyHandle->NotifyPreChange();
 		}
-		TargetSettings->UpdateWordCountsFromCSV();
-		TargetSettings->UpdateStatusFromConflictReport();
+		LocalizationTarget->UpdateWordCountsFromCSV();
+		LocalizationTarget->UpdateStatusFromConflictReport();
 		for (const TSharedPtr<IPropertyHandle>& WordCountPropertyHandle : WordCountPropertyHandles)
 		{
 			WordCountPropertyHandle->NotifyPostChange();
@@ -406,8 +417,8 @@ void SLocalizationTargetEditorCultureRow::Delete()
 		TGuardValue<bool> ReentranceGuard(IsExecuting, true);
 
 		const FCulturePtr Culture = GetCulture();
-		FLocalizationTargetSettings* const TargetSettings = GetTargetSettings();
-		if (Culture.IsValid() && TargetSettings)
+		ULocalizationTarget* const LocalizationTarget = GetTarget();
+		if (Culture.IsValid() && LocalizationTarget)
 		{
 			FText TitleText;
 			FText MessageText;
@@ -424,7 +435,7 @@ void SLocalizationTargetEditorCultureRow::Delete()
 			case EAppReturnType::Ok:
 				{
 					const FString CultureName = Culture->GetName();
-					TargetSettings->DeleteFiles(&CultureName);
+					LocalizationTarget->DeleteFiles(&CultureName);
 
 					// Remove this element from the parent array.
 					const TSharedPtr<IPropertyHandle> SupportedCulturesStatisticsPropertyHandle = TargetSettingsPropertyHandle.IsValid() && TargetSettingsPropertyHandle->IsValidHandle() ? TargetSettingsPropertyHandle->GetChildHandle(GET_MEMBER_NAME_CHECKED(FLocalizationTargetSettings, SupportedCulturesStatistics)) : nullptr;
diff --git a/Engine/Source/Editor/LocalizationDashboard/Private/SLocalizationTargetEditorCultureRow.h b/Engine/Source/Editor/LocalizationDashboard/Private/SLocalizationTargetEditorCultureRow.h
--- a/Engine/Source/Editor/LocalizationDashboard/Private/SLocalizationTargetEditorCultureRow.h
+++ b/Engine/Source/Editor/LocalizationDashboard/Private/SLocalizationTargetEditorCultureRow.h
@@ -13,6 +13,7 @@ public:
 	TSharedRef<SWidget> GenerateWidgetForColumn( const FName& ColumnName ) override;
 
 private:
+	ULocalizationTarget* GetTarget() const;
 	FLocalizationTargetSettings* GetTargetSettings() const;
 	FCultureStatistics* GetCultureStatistics() const;
 	FCulturePtr GetCulture() const;
diff --git a/Engine/Source/Editor/LocalizationDashboard/Private/SLocalizationTargetStatusButton.cpp b/Engine/Source/Editor/LocalizationDashboard/Private/SLocalizationTargetStatusButton.cpp
--- a/Engine/Source/Editor/LocalizationDashboard/Private/SLocalizationTargetStatusButton.cpp
+++ b/Engine/Source/Editor/LocalizationDashboard/Private/SLocalizationTargetStatusButton.cpp
@@ -89,7 +89,7 @@ FReply SLocalizationTargetStatusButton::OnClicked()
 			const TSharedPtr<SWindow> ParentWindow = FSlateApplication::Get().FindWidgetWindow(AsShared());
 			if (LocalizationCommandletTasks::GenerateReportsForTarget(ParentWindow.ToSharedRef(), Target->Settings))
 			{
-				Target->Settings.UpdateStatusFromConflictReport();
+				Target->UpdateStatusFromConflictReport();
 			}
 		}
 		break;
