Refactored localization commandlet execution/tasks so that they take the ULocalizationTarget object instead of the FLocalizationTargetSettings member in preparation for having the UE4 target set separate from the project-specific target set.

diff --git a/Engine/Source/Editor/LocalizationDashboard/Private/LocalizationCommandletTasks.cpp b/Engine/Source/Editor/LocalizationDashboard/Private/LocalizationCommandletTasks.cpp
--- a/Engine/Source/Editor/LocalizationDashboard/Private/LocalizationCommandletTasks.cpp
+++ b/Engine/Source/Editor/LocalizationDashboard/Private/LocalizationCommandletTasks.cpp
@@ -7,93 +7,93 @@
 
 #define LOCTEXT_NAMESPACE "LocalizationCommandletTasks"
 
-bool LocalizationCommandletTasks::GatherTargets(const TSharedRef<SWindow>& ParentWindow, const TArray<FLocalizationTargetSettings*>& TargetsSettings)
+bool LocalizationCommandletTasks::GatherTargets(const TSharedRef<SWindow>& ParentWindow, const TArray<ULocalizationTarget*>& Targets)
 {
 	TArray<LocalizationCommandletExecution::FTask> Tasks;
 
-	for (FLocalizationTargetSettings* TargetSettings : TargetsSettings)
+	for (ULocalizationTarget* Target : Targets)
 	{
 		FFormatNamedArguments Arguments;
-		Arguments.Add(TEXT("TargetName"), FText::FromString(TargetSettings->Name));
+		Arguments.Add(TEXT("TargetName"), FText::FromString(Target->Settings.Name));
 
 		const FText GatherTaskName = FText::Format(LOCTEXT("GatherTaskNameFormat", "Gather Text for {TargetName}"), Arguments);
-		const FString GatherScriptPath = LocalizationConfigurationScript::GetGatherScriptPath(*TargetSettings);
-		LocalizationConfigurationScript::GenerateGatherScript(*TargetSettings).Write(GatherScriptPath);
+		const FString GatherScriptPath = LocalizationConfigurationScript::GetGatherScriptPath(Target);
+		LocalizationConfigurationScript::GenerateGatherScript(Target).Write(GatherScriptPath);
 		Tasks.Add(LocalizationCommandletExecution::FTask(GatherTaskName, GatherScriptPath));
 
 		const FText ReportTaskName = FText::Format(LOCTEXT("ReportTaskNameFormat", "Generate Reports for {TargetName}"), Arguments);
-		const FString ReportScriptPath = LocalizationConfigurationScript::GetReportScriptPath(*TargetSettings);
-		LocalizationConfigurationScript::GenerateReportScript(*TargetSettings).Write(ReportScriptPath);
+		const FString ReportScriptPath = LocalizationConfigurationScript::GetReportScriptPath(Target);
+		LocalizationConfigurationScript::GenerateReportScript(Target).Write(ReportScriptPath);
 		Tasks.Add(LocalizationCommandletExecution::FTask(ReportTaskName, ReportScriptPath));
 	}
 
 	return LocalizationCommandletExecution::Execute(ParentWindow, LOCTEXT("GatherAllTargetsWindowTitle", "Gather Text for All Targets"), Tasks);
 }
 
-bool LocalizationCommandletTasks::GatherTarget(const TSharedRef<SWindow>& ParentWindow, FLocalizationTargetSettings& TargetSettings)
+bool LocalizationCommandletTasks::GatherTarget(const TSharedRef<SWindow>& ParentWindow, ULocalizationTarget* const Target)
 {
 	TArray<LocalizationCommandletExecution::FTask> Tasks;
 
-	const FString GatherScriptPath = LocalizationConfigurationScript::GetGatherScriptPath(TargetSettings);
-	LocalizationConfigurationScript::GenerateGatherScript(TargetSettings).Write(GatherScriptPath);
+	const FString GatherScriptPath = LocalizationConfigurationScript::GetGatherScriptPath(Target);
+	LocalizationConfigurationScript::GenerateGatherScript(Target).Write(GatherScriptPath);
 	Tasks.Add(LocalizationCommandletExecution::FTask(LOCTEXT("GatherTaskName", "Gather Text"), GatherScriptPath));
 
-	const FString ReportScriptPath = LocalizationConfigurationScript::GetReportScriptPath(TargetSettings);
-	LocalizationConfigurationScript::GenerateReportScript(TargetSettings).Write(ReportScriptPath);
+	const FString ReportScriptPath = LocalizationConfigurationScript::GetReportScriptPath(Target);
+	LocalizationConfigurationScript::GenerateReportScript(Target).Write(ReportScriptPath);
 	Tasks.Add(LocalizationCommandletExecution::FTask(LOCTEXT("ReportTaskName", "Generate Reports"), ReportScriptPath));
 
 	FFormatNamedArguments Arguments;
-	Arguments.Add(TEXT("TargetName"), FText::FromString(TargetSettings.Name));
+	Arguments.Add(TEXT("TargetName"), FText::FromString(Target->Settings.Name));
 
 	const FText WindowTitle = FText::Format(LOCTEXT("GatherTargetWindowTitle", "Gather Text for Target {TargetName}"), Arguments);
 	return LocalizationCommandletExecution::Execute(ParentWindow, WindowTitle, Tasks);
 }
 
 
-bool LocalizationCommandletTasks::ImportTargets(const TSharedRef<SWindow>& ParentWindow, const TArray<FLocalizationTargetSettings*>& TargetsSettings, const TOptional<FString> DirectoryPath)
+bool LocalizationCommandletTasks::ImportTargets(const TSharedRef<SWindow>& ParentWindow, const TArray<ULocalizationTarget*>& Targets, const TOptional<FString> DirectoryPath)
 {
 	TArray<LocalizationCommandletExecution::FTask> Tasks;
 
-	for (FLocalizationTargetSettings* TargetSettings : TargetsSettings)
+	for (ULocalizationTarget* Target : Targets)
 	{
 		FFormatNamedArguments Arguments;
-		Arguments.Add(TEXT("TargetName"), FText::FromString(TargetSettings->Name));
+		Arguments.Add(TEXT("TargetName"), FText::FromString(Target->Settings.Name));
 
 		const FText ImportTaskName = FText::Format(LOCTEXT("ImportTaskNameFormat", "Import Translations for {TargetName}"), Arguments);
-		const FString ImportScriptPath = LocalizationConfigurationScript::GetImportScriptPath(*TargetSettings, TOptional<FString>());
-		const TOptional<FString> DirectoryPathForTarget = DirectoryPath.IsSet() ? DirectoryPath.GetValue() / TargetSettings->Name : TOptional<FString>();
-		LocalizationConfigurationScript::GenerateImportScript(*TargetSettings, TOptional<FString>(), DirectoryPathForTarget).Write(ImportScriptPath);
+		const FString ImportScriptPath = LocalizationConfigurationScript::GetImportScriptPath(Target, TOptional<FString>());
+		const TOptional<FString> DirectoryPathForTarget = DirectoryPath.IsSet() ? DirectoryPath.GetValue() / Target->Settings.Name : TOptional<FString>();
+		LocalizationConfigurationScript::GenerateImportScript(Target, TOptional<FString>(), DirectoryPathForTarget).Write(ImportScriptPath);
 		Tasks.Add(LocalizationCommandletExecution::FTask(ImportTaskName, ImportScriptPath));
 
 		const FText ReportTaskName = FText::Format(LOCTEXT("ReportTaskNameFormat", "Generate Reports for {TargetName}"), Arguments);
-		const FString ReportScriptPath = LocalizationConfigurationScript::GetReportScriptPath(*TargetSettings);
-		LocalizationConfigurationScript::GenerateReportScript(*TargetSettings).Write(ReportScriptPath);
+		const FString ReportScriptPath = LocalizationConfigurationScript::GetReportScriptPath(Target);
+		LocalizationConfigurationScript::GenerateReportScript(Target).Write(ReportScriptPath);
 		Tasks.Add(LocalizationCommandletExecution::FTask(ReportTaskName, ReportScriptPath));
 	}
 
 	return LocalizationCommandletExecution::Execute(ParentWindow, LOCTEXT("ImportForAllTargetsWindowTitle", "Import Translations for All Targets"), Tasks);
 }
 
-bool LocalizationCommandletTasks::ImportTarget(const TSharedRef<SWindow>& ParentWindow, FLocalizationTargetSettings& TargetSettings, const TOptional<FString> DirectoryPath)
+bool LocalizationCommandletTasks::ImportTarget(const TSharedRef<SWindow>& ParentWindow, ULocalizationTarget* const Target, const TOptional<FString> DirectoryPath)
 {
 	TArray<LocalizationCommandletExecution::FTask> Tasks;
 
-	const FString ImportScriptPath = LocalizationConfigurationScript::GetImportScriptPath(TargetSettings, TOptional<FString>());
-	LocalizationConfigurationScript::GenerateImportScript(TargetSettings, TOptional<FString>(), DirectoryPath).Write(ImportScriptPath);
+	const FString ImportScriptPath = LocalizationConfigurationScript::GetImportScriptPath(Target, TOptional<FString>());
+	LocalizationConfigurationScript::GenerateImportScript(Target, TOptional<FString>(), DirectoryPath).Write(ImportScriptPath);
 	Tasks.Add(LocalizationCommandletExecution::FTask(LOCTEXT("ImportTaskName", "Import Translations"), ImportScriptPath));
 
-	const FString ReportScriptPath = LocalizationConfigurationScript::GetReportScriptPath(TargetSettings);
-	LocalizationConfigurationScript::GenerateReportScript(TargetSettings).Write(ReportScriptPath);
+	const FString ReportScriptPath = LocalizationConfigurationScript::GetReportScriptPath(Target);
+	LocalizationConfigurationScript::GenerateReportScript(Target).Write(ReportScriptPath);
 	Tasks.Add(LocalizationCommandletExecution::FTask(LOCTEXT("ReportTaskName", "Generate Reports"), ReportScriptPath));
 
 	FFormatNamedArguments Arguments;
-	Arguments.Add(TEXT("TargetName"), FText::FromString(TargetSettings.Name));
+	Arguments.Add(TEXT("TargetName"), FText::FromString(Target->Settings.Name));
 
 	const FText WindowTitle = FText::Format(LOCTEXT("ImportForTargetWindowTitle", "Import Translations for Target {TargetName}"), Arguments);
 	return LocalizationCommandletExecution::Execute(ParentWindow, WindowTitle, Tasks);
 }
 
-bool LocalizationCommandletTasks::ImportCulture(const TSharedRef<SWindow>& ParentWindow, FLocalizationTargetSettings& TargetSettings, const FString& CultureName, const TOptional<FString> FilePath)
+bool LocalizationCommandletTasks::ImportCulture(const TSharedRef<SWindow>& ParentWindow, ULocalizationTarget* const Target, const FString& CultureName, const TOptional<FString> FilePath)
 {
 	FCulturePtr Culture = FInternationalization::Get().GetCulture(CultureName);
 	if (!Culture.IsValid())
@@ -103,18 +103,18 @@ bool LocalizationCommandletTasks::ImportCulture(const TSharedRef<SWindow>& Paren
 
 	TArray<LocalizationCommandletExecution::FTask> Tasks;
 
-	const FString DefaultImportScriptPath = LocalizationConfigurationScript::GetImportScriptPath(TargetSettings, TOptional<FString>(CultureName));
+	const FString DefaultImportScriptPath = LocalizationConfigurationScript::GetImportScriptPath(Target, TOptional<FString>(CultureName));
 	const FString ImportScriptPath = FPaths::CreateTempFilename(*FPaths::GetPath(DefaultImportScriptPath), *FPaths::GetBaseFilename(DefaultImportScriptPath), *FPaths::GetExtension(DefaultImportScriptPath, true));
-	LocalizationConfigurationScript::GenerateImportScript(TargetSettings, TOptional<FString>(CultureName), FilePath).Write(ImportScriptPath);
+	LocalizationConfigurationScript::GenerateImportScript(Target, TOptional<FString>(CultureName), FilePath).Write(ImportScriptPath);
 	Tasks.Add(LocalizationCommandletExecution::FTask(LOCTEXT("ImportTaskName", "Import Translations"), ImportScriptPath));
 
-	const FString ReportScriptPath = LocalizationConfigurationScript::GetReportScriptPath(TargetSettings);
-	LocalizationConfigurationScript::GenerateReportScript(TargetSettings).Write(ReportScriptPath);
+	const FString ReportScriptPath = LocalizationConfigurationScript::GetReportScriptPath(Target);
+	LocalizationConfigurationScript::GenerateReportScript(Target).Write(ReportScriptPath);
 	Tasks.Add(LocalizationCommandletExecution::FTask(LOCTEXT("ReportTaskName", "Generate Reports"), ReportScriptPath));
 
 	FFormatNamedArguments Arguments;
 	Arguments.Add(TEXT("CultureName"), FText::FromString(Culture->GetDisplayName()));
-	Arguments.Add(TEXT("TargetName"), FText::FromString(TargetSettings.Name));
+	Arguments.Add(TEXT("TargetName"), FText::FromString(Target->Settings.Name));
 
 	const FText WindowTitle = FText::Format(LOCTEXT("ImportCultureForTargetWindowTitle", "Import {CultureName} Translations for Target {TargetName}"), Arguments);
 
@@ -124,41 +124,41 @@ bool LocalizationCommandletTasks::ImportCulture(const TSharedRef<SWindow>& Paren
 
 }
 
-bool LocalizationCommandletTasks::ExportTargets(const TSharedRef<SWindow>& ParentWindow, const TArray<FLocalizationTargetSettings*>& TargetsSettings, const TOptional<FString> DirectoryPath)
+bool LocalizationCommandletTasks::ExportTargets(const TSharedRef<SWindow>& ParentWindow, const TArray<ULocalizationTarget*>& Targets, const TOptional<FString> DirectoryPath)
 {
 	TArray<LocalizationCommandletExecution::FTask> Tasks;
 
-	for (FLocalizationTargetSettings* TargetSettings : TargetsSettings)
+	for (ULocalizationTarget* Target : Targets)
 	{
 		FFormatNamedArguments Arguments;
-		Arguments.Add(TEXT("TargetName"), FText::FromString(TargetSettings->Name));
+		Arguments.Add(TEXT("TargetName"), FText::FromString(Target->Settings.Name));
 
 		const FText ExportTaskName = FText::Format(LOCTEXT("ExportTaskNameFormat", "Export Translations for {TargetName}"), Arguments);
-		const FString ExportScriptPath = LocalizationConfigurationScript::GetExportScriptPath(*TargetSettings, TOptional<FString>());
-		const TOptional<FString> DirectoryPathForTarget = DirectoryPath.IsSet() ? DirectoryPath.GetValue() / TargetSettings->Name : TOptional<FString>();
-		LocalizationConfigurationScript::GenerateExportScript(*TargetSettings, TOptional<FString>(), DirectoryPathForTarget).Write(ExportScriptPath);
+		const FString ExportScriptPath = LocalizationConfigurationScript::GetExportScriptPath(Target, TOptional<FString>());
+		const TOptional<FString> DirectoryPathForTarget = DirectoryPath.IsSet() ? DirectoryPath.GetValue() / Target->Settings.Name : TOptional<FString>();
+		LocalizationConfigurationScript::GenerateExportScript(Target, TOptional<FString>(), DirectoryPathForTarget).Write(ExportScriptPath);
 		Tasks.Add(LocalizationCommandletExecution::FTask(ExportTaskName, ExportScriptPath));
 	}
 
 	return LocalizationCommandletExecution::Execute(ParentWindow, LOCTEXT("ExportForAllTargetsWindowTitle", "Export Translations for All Targets"), Tasks);
 }
 
-bool LocalizationCommandletTasks::ExportTarget(const TSharedRef<SWindow>& ParentWindow, FLocalizationTargetSettings& TargetSettings, const TOptional<FString> DirectoryPath)
+bool LocalizationCommandletTasks::ExportTarget(const TSharedRef<SWindow>& ParentWindow, ULocalizationTarget* const Target, const TOptional<FString> DirectoryPath)
 {
 	TArray<LocalizationCommandletExecution::FTask> Tasks;
 
-	const FString ExportScriptPath = LocalizationConfigurationScript::GetExportScriptPath(TargetSettings, TOptional<FString>());
-	LocalizationConfigurationScript::GenerateExportScript(TargetSettings, TOptional<FString>(), DirectoryPath).Write(ExportScriptPath);
+	const FString ExportScriptPath = LocalizationConfigurationScript::GetExportScriptPath(Target, TOptional<FString>());
+	LocalizationConfigurationScript::GenerateExportScript(Target, TOptional<FString>(), DirectoryPath).Write(ExportScriptPath);
 	Tasks.Add(LocalizationCommandletExecution::FTask(LOCTEXT("ExportTaskName", "Export Translations"), ExportScriptPath));
 
 	FFormatNamedArguments Arguments;
-	Arguments.Add(TEXT("TargetName"), FText::FromString(TargetSettings.Name));
+	Arguments.Add(TEXT("TargetName"), FText::FromString(Target->Settings.Name));
 
 	const FText WindowTitle = FText::Format(LOCTEXT("ExportForTargetWindowTitle", "Export Translations for Target {TargetName}"), Arguments);
 	return LocalizationCommandletExecution::Execute(ParentWindow, WindowTitle, Tasks);
 }
 
-bool LocalizationCommandletTasks::ExportCulture(const TSharedRef<SWindow>& ParentWindow, FLocalizationTargetSettings& TargetSettings, const FString& CultureName, const TOptional<FString> FilePath)
+bool LocalizationCommandletTasks::ExportCulture(const TSharedRef<SWindow>& ParentWindow, ULocalizationTarget* const Target, const FString& CultureName, const TOptional<FString> FilePath)
 {
 	FCulturePtr Culture = FInternationalization::Get().GetCulture(CultureName);
 	if (!Culture.IsValid())
@@ -168,14 +168,14 @@ bool LocalizationCommandletTasks::ExportCulture(const TSharedRef<SWindow>& Paren
 
 	TArray<LocalizationCommandletExecution::FTask> Tasks;
 
-	const FString DefaultExportScriptPath = LocalizationConfigurationScript::GetExportScriptPath(TargetSettings, TOptional<FString>(CultureName));
+	const FString DefaultExportScriptPath = LocalizationConfigurationScript::GetExportScriptPath(Target, TOptional<FString>(CultureName));
 	const FString ExportScriptPath = FPaths::CreateTempFilename(*FPaths::GetPath(DefaultExportScriptPath), *FPaths::GetBaseFilename(DefaultExportScriptPath), *FPaths::GetExtension(DefaultExportScriptPath, true));
-	LocalizationConfigurationScript::GenerateExportScript(TargetSettings, TOptional<FString>(CultureName), FilePath).Write(ExportScriptPath);
+	LocalizationConfigurationScript::GenerateExportScript(Target, TOptional<FString>(CultureName), FilePath).Write(ExportScriptPath);
 	Tasks.Add(LocalizationCommandletExecution::FTask(LOCTEXT("ExportTaskName", "Export Translations"), ExportScriptPath));
 
 	FFormatNamedArguments Arguments;
 	Arguments.Add(TEXT("CultureName"), FText::FromString(Culture->GetDisplayName()));
-	Arguments.Add(TEXT("TargetName"), FText::FromString(TargetSettings.Name));
+	Arguments.Add(TEXT("TargetName"), FText::FromString(Target->Settings.Name));
 
 	const FText WindowTitle = FText::Format(LOCTEXT("ExportCultureForTargetWindowTitle", "Export {CultureName} Translations for Target {TargetName}"), Arguments);
 
@@ -184,67 +184,67 @@ bool LocalizationCommandletTasks::ExportCulture(const TSharedRef<SWindow>& Paren
 	return HasSucceeeded;
 }
 
-bool LocalizationCommandletTasks::GenerateReportsForTargets(const TSharedRef<SWindow>& ParentWindow, const TArray<FLocalizationTargetSettings*>& TargetsSettings)
+bool LocalizationCommandletTasks::GenerateReportsForTargets(const TSharedRef<SWindow>& ParentWindow, const TArray<ULocalizationTarget*>& Targets)
 {
 	TArray<LocalizationCommandletExecution::FTask> Tasks;
 
-	for (FLocalizationTargetSettings* TargetSettings : TargetsSettings)
+	for (ULocalizationTarget* Target : Targets)
 	{
 		FFormatNamedArguments Arguments;
-		Arguments.Add(TEXT("TargetName"), FText::FromString(TargetSettings->Name));
+		Arguments.Add(TEXT("TargetName"), FText::FromString(Target->Settings.Name));
 
 		const FText ReportTaskName = FText::Format(LOCTEXT("ReportTaskNameFormat", "Generate Reports for {TargetName}"), Arguments);
-		const FString ReportScriptPath = LocalizationConfigurationScript::GetReportScriptPath(*TargetSettings);
-		LocalizationConfigurationScript::GenerateReportScript(*TargetSettings).Write(ReportScriptPath);
+		const FString ReportScriptPath = LocalizationConfigurationScript::GetReportScriptPath(Target);
+		LocalizationConfigurationScript::GenerateReportScript(Target).Write(ReportScriptPath);
 		Tasks.Add(LocalizationCommandletExecution::FTask(ReportTaskName, ReportScriptPath));
 	}
 
 	return LocalizationCommandletExecution::Execute(ParentWindow, LOCTEXT("GenerateReportsForAllTargetsWindowTitle", "Generate Reports for All Targets"), Tasks);
 }
 
-bool LocalizationCommandletTasks::GenerateReportsForTarget(const TSharedRef<SWindow>& ParentWindow, FLocalizationTargetSettings& TargetSettings)
+bool LocalizationCommandletTasks::GenerateReportsForTarget(const TSharedRef<SWindow>& ParentWindow, ULocalizationTarget* const Target)
 {
 	TArray<LocalizationCommandletExecution::FTask> Tasks;
 
-	const FString ReportScriptPath = LocalizationConfigurationScript::GetReportScriptPath(TargetSettings);
-	LocalizationConfigurationScript::GenerateReportScript(TargetSettings).Write(ReportScriptPath);
+	const FString ReportScriptPath = LocalizationConfigurationScript::GetReportScriptPath(Target);
+	LocalizationConfigurationScript::GenerateReportScript(Target).Write(ReportScriptPath);
 	Tasks.Add(LocalizationCommandletExecution::FTask(LOCTEXT("ReportTaskName", "Generate Reports"), ReportScriptPath));
 
 	FFormatNamedArguments Arguments;
-	Arguments.Add(TEXT("TargetName"), FText::FromString(TargetSettings.Name));
+	Arguments.Add(TEXT("TargetName"), FText::FromString(Target->Settings.Name));
 
 	const FText WindowTitle = FText::Format(LOCTEXT("GenerateReportsForTargetWindowTitle", "Generate Reports for Target {TargetName}"), Arguments);
 	return LocalizationCommandletExecution::Execute(ParentWindow, WindowTitle, Tasks);
 }
 
-bool LocalizationCommandletTasks::CompileTargets(const TSharedRef<SWindow>& ParentWindow, const TArray<FLocalizationTargetSettings*>& TargetsSettings)
+bool LocalizationCommandletTasks::CompileTargets(const TSharedRef<SWindow>& ParentWindow, const TArray<ULocalizationTarget*>& Targets)
 {
 	TArray<LocalizationCommandletExecution::FTask> Tasks;
 
-	for (FLocalizationTargetSettings* TargetSettings : TargetsSettings)
+	for (ULocalizationTarget* Target : Targets)
 	{
 		FFormatNamedArguments Arguments;
-		Arguments.Add(TEXT("TargetName"), FText::FromString(TargetSettings->Name));
+		Arguments.Add(TEXT("TargetName"), FText::FromString(Target->Settings.Name));
 
 		const FText CompileTaskName = FText::Format(LOCTEXT("CompileTaskNameFormat", "Compile Translations for {TargetName}"), Arguments);
-		const FString CompileScriptPath = LocalizationConfigurationScript::GetCompileScriptPath(*TargetSettings);
-		LocalizationConfigurationScript::GenerateCompileScript(*TargetSettings).Write(CompileScriptPath);
+		const FString CompileScriptPath = LocalizationConfigurationScript::GetCompileScriptPath(Target);
+		LocalizationConfigurationScript::GenerateCompileScript(Target).Write(CompileScriptPath);
 		Tasks.Add(LocalizationCommandletExecution::FTask(CompileTaskName, CompileScriptPath));
 	}
 
 	return LocalizationCommandletExecution::Execute(ParentWindow, LOCTEXT("GenerateLocResForAllTargetsWindowTitle", "Compile Translations for All Targets"), Tasks);
 }
 
-bool LocalizationCommandletTasks::CompileTarget(const TSharedRef<SWindow>& ParentWindow, FLocalizationTargetSettings& TargetSettings)
+bool LocalizationCommandletTasks::CompileTarget(const TSharedRef<SWindow>& ParentWindow, ULocalizationTarget* const Target)
 {
 	TArray<LocalizationCommandletExecution::FTask> Tasks;
 
-	const FString CompileScriptPath = LocalizationConfigurationScript::GetCompileScriptPath(TargetSettings);
-	LocalizationConfigurationScript::GenerateCompileScript(TargetSettings).Write(CompileScriptPath);
+	const FString CompileScriptPath = LocalizationConfigurationScript::GetCompileScriptPath(Target);
+	LocalizationConfigurationScript::GenerateCompileScript(Target).Write(CompileScriptPath);
 	Tasks.Add(LocalizationCommandletExecution::FTask(LOCTEXT("CompileTaskName", "Compile Translations"), CompileScriptPath));
 
 	FFormatNamedArguments Arguments;
-	Arguments.Add(TEXT("TargetName"), FText::FromString(TargetSettings.Name));
+	Arguments.Add(TEXT("TargetName"), FText::FromString(Target->Settings.Name));
 
 	const FText WindowTitle = FText::Format(LOCTEXT("GenerateLocResForTargetWindowTitle", "Compile Translations for Target {TargetName}"), Arguments);
 	return LocalizationCommandletExecution::Execute(ParentWindow, WindowTitle, Tasks);
diff --git a/Engine/Source/Editor/LocalizationDashboard/Private/LocalizationCommandletTasks.h b/Engine/Source/Editor/LocalizationDashboard/Private/LocalizationCommandletTasks.h
--- a/Engine/Source/Editor/LocalizationDashboard/Private/LocalizationCommandletTasks.h
+++ b/Engine/Source/Editor/LocalizationDashboard/Private/LocalizationCommandletTasks.h
@@ -5,20 +5,20 @@
 
 namespace LocalizationCommandletTasks
 {
-	bool GatherTargets(const TSharedRef<SWindow>& ParentWindow, const TArray<FLocalizationTargetSettings*>& TargetsSettings);
-	bool GatherTarget(const TSharedRef<SWindow>& ParentWindow, FLocalizationTargetSettings& TargetSettings);
+	bool GatherTargets(const TSharedRef<SWindow>& ParentWindow, const TArray<ULocalizationTarget*>& Targets);
+	bool GatherTarget(const TSharedRef<SWindow>& ParentWindow, ULocalizationTarget* const Target);
 
-	bool ImportTargets(const TSharedRef<SWindow>& ParentWindow, const TArray<FLocalizationTargetSettings*>& TargetsSettings, const TOptional<FString> DirectoryPath = TOptional<FString>());
-	bool ImportTarget(const TSharedRef<SWindow>& ParentWindow, FLocalizationTargetSettings& TargetSettings, const TOptional<FString> DirectoryPath = TOptional<FString>());
-	bool ImportCulture(const TSharedRef<SWindow>& ParentWindow, FLocalizationTargetSettings& TargetSettings, const FString& CultureName, const TOptional<FString> FilePath = TOptional<FString>());
+	bool ImportTargets(const TSharedRef<SWindow>& ParentWindow, const TArray<ULocalizationTarget*>& Targets, const TOptional<FString> DirectoryPath = TOptional<FString>());
+	bool ImportTarget(const TSharedRef<SWindow>& ParentWindow, ULocalizationTarget* const Target, const TOptional<FString> DirectoryPath = TOptional<FString>());
+	bool ImportCulture(const TSharedRef<SWindow>& ParentWindow, ULocalizationTarget* const Target, const FString& CultureName, const TOptional<FString> FilePath = TOptional<FString>());
 
-	bool ExportTargets(const TSharedRef<SWindow>& ParentWindow, const TArray<FLocalizationTargetSettings*>& TargetsSettings, const TOptional<FString> DirectoryPath = TOptional<FString>());
-	bool ExportTarget(const TSharedRef<SWindow>& ParentWindow, FLocalizationTargetSettings& TargetSettings, const TOptional<FString> DirectoryPath = TOptional<FString>());
-	bool ExportCulture(const TSharedRef<SWindow>& ParentWindow, FLocalizationTargetSettings& TargetSettings, const FString& CultureName, const TOptional<FString> FilePath = TOptional<FString>());
+	bool ExportTargets(const TSharedRef<SWindow>& ParentWindow, const TArray<ULocalizationTarget*>& Targets, const TOptional<FString> DirectoryPath = TOptional<FString>());
+	bool ExportTarget(const TSharedRef<SWindow>& ParentWindow, ULocalizationTarget* const Target, const TOptional<FString> DirectoryPath = TOptional<FString>());
+	bool ExportCulture(const TSharedRef<SWindow>& ParentWindow, ULocalizationTarget* const Target, const FString& CultureName, const TOptional<FString> FilePath = TOptional<FString>());
 
-	bool GenerateReportsForTargets(const TSharedRef<SWindow>& ParentWindow, const TArray<FLocalizationTargetSettings*>& TargetsSettings);
-	bool GenerateReportsForTarget(const TSharedRef<SWindow>& ParentWindow, FLocalizationTargetSettings& TargetSettings);
+	bool GenerateReportsForTargets(const TSharedRef<SWindow>& ParentWindow, const TArray<ULocalizationTarget*>& Targets);
+	bool GenerateReportsForTarget(const TSharedRef<SWindow>& ParentWindow, ULocalizationTarget* const Target);
 
-	bool CompileTargets(const TSharedRef<SWindow>& ParentWindow, const TArray<FLocalizationTargetSettings*>& TargetsSettings);
-	bool CompileTarget(const TSharedRef<SWindow>& ParentWindow, FLocalizationTargetSettings& TargetSettings);
+	bool CompileTargets(const TSharedRef<SWindow>& ParentWindow, const TArray<ULocalizationTarget*>& Targets);
+	bool CompileTarget(const TSharedRef<SWindow>& ParentWindow, ULocalizationTarget* const Target);
 }
\ No newline at end of file
diff --git a/Engine/Source/Editor/LocalizationDashboard/Private/LocalizationConfigurationScript.cpp b/Engine/Source/Editor/LocalizationDashboard/Private/LocalizationConfigurationScript.cpp
--- a/Engine/Source/Editor/LocalizationDashboard/Private/LocalizationConfigurationScript.cpp
+++ b/Engine/Source/Editor/LocalizationDashboard/Private/LocalizationConfigurationScript.cpp
@@ -18,12 +18,12 @@ namespace LocalizationConfigurationScript
 		return FPaths::GameConfigDir() / TEXT("Localization");
 	}
 
-	FString GetDataDirectory(const FLocalizationTargetSettings& Target)
+	FString GetDataDirectory(const ULocalizationTarget* const Target)
 	{
-		return FPaths::GameContentDir() / TEXT("Localization") / Target.Name;
+		return FPaths::GameContentDir() / TEXT("Localization") / Target->Settings.Name;
 	}
 
-	TArray<FString> GetScriptPaths(const FLocalizationTargetSettings& Target)
+	TArray<FString> GetScriptPaths(const ULocalizationTarget* const Target)
 	{
 		TArray<FString> Result;
 		Result.Add(GetGatherScriptPath(Target));
@@ -33,42 +33,42 @@ namespace LocalizationConfigurationScript
 		return Result;
 	}
 
-	FString GetManifestPath(const FLocalizationTargetSettings& Target)
+	FString GetManifestPath(const ULocalizationTarget* const Target)
 	{
-		return GetDataDirectory(Target) / FString::Printf( TEXT("%s.%s"), *Target.Name, TEXT("manifest") );
+		return GetDataDirectory(Target) / FString::Printf( TEXT("%s.%s"), *Target->Settings.Name, TEXT("manifest") );
 	}
 
-	FString GetArchivePath(const FLocalizationTargetSettings& Target, const FString& CultureName)
+	FString GetArchivePath(const ULocalizationTarget* const Target, const FString& CultureName)
 	{
-		return GetDataDirectory(Target) / CultureName / FString::Printf( TEXT("%s.%s"), *Target.Name, TEXT("archive") );
+		return GetDataDirectory(Target) / CultureName / FString::Printf( TEXT("%s.%s"), *Target->Settings.Name, TEXT("archive") );
 	}
 
-	FString GetDefaultPOFileName(const FLocalizationTargetSettings& Target)
+	FString GetDefaultPOFileName(const ULocalizationTarget* const Target)
 	{
-		return FString::Printf( TEXT("%s.%s"), *Target.Name, TEXT("po") );
+		return FString::Printf( TEXT("%s.%s"), *Target->Settings.Name, TEXT("po") );
 	}
 
-	FString GetDefaultPOPath(const FLocalizationTargetSettings& Target, const FString& CultureName)
+	FString GetDefaultPOPath(const ULocalizationTarget* const Target, const FString& CultureName)
 	{
 		return GetDataDirectory(Target) / CultureName / GetDefaultPOFileName(Target);
 	}
 
-	FString GetLocResPath(const FLocalizationTargetSettings& Target, const FString& CultureName)
+	FString GetLocResPath(const ULocalizationTarget* const Target, const FString& CultureName)
 	{
-		return GetDataDirectory(Target) / CultureName / FString::Printf( TEXT("%s.%s"), *Target.Name, TEXT("locres") );
+		return GetDataDirectory(Target) / CultureName / FString::Printf( TEXT("%s.%s"), *Target->Settings.Name, TEXT("locres") );
 	}
 
-	FString GetWordCountCSVPath(const FLocalizationTargetSettings& Target)
+	FString GetWordCountCSVPath(const ULocalizationTarget* const Target)
 	{
-		return GetDataDirectory(Target) / FString::Printf( TEXT("%s.%s"), *Target.Name, TEXT("csv") );
+		return GetDataDirectory(Target) / FString::Printf( TEXT("%s.%s"), *Target->Settings.Name, TEXT("csv") );
 	}
 
-	FString GetConflictReportPath(const FLocalizationTargetSettings& Target)
+	FString GetConflictReportPath(const ULocalizationTarget* const Target)
 	{
-		return GetDataDirectory(Target) / FString::Printf( TEXT("%s_Conflicts.%s"), *Target.Name, TEXT("txt") );
+		return GetDataDirectory(Target) / FString::Printf( TEXT("%s_Conflicts.%s"), *Target->Settings.Name, TEXT("txt") );
 	}
 
-	FLocalizationConfigurationScript GenerateGatherScript(const FLocalizationTargetSettings& Target)
+	FLocalizationConfigurationScript GenerateGatherScript(const ULocalizationTarget* const Target)
 	{
 		FLocalizationConfigurationScript Script;
 
@@ -83,31 +83,31 @@ namespace LocalizationConfigurationScript
 			FConfigSection& ConfigSection = Script.CommonSettings();
 
 			const ULocalizationTargetSet* const LocalizationTargetSet = GetDefault<ULocalizationTargetSet>(ULocalizationTargetSet::StaticClass());
-			for (const FString& TargetDependencyName : Target.TargetDependencies)
+			for (const FString& TargetDependencyName : Target->Settings.TargetDependencies)
 			{
 				ULocalizationTarget* const * OtherTarget = LocalizationTargetSet->TargetObjects.FindByPredicate([&TargetDependencyName](ULocalizationTarget* const OtherTarget)->bool{return OtherTarget->Settings.Name == TargetDependencyName;});
 				if (OtherTarget)
 				{
-					ConfigSection.Add( TEXT("ManifestDependencies"), MakePathRelativeToProjectDirectory(GetManifestPath((*OtherTarget)->Settings)) );
+					ConfigSection.Add( TEXT("ManifestDependencies"), MakePathRelativeToProjectDirectory(GetManifestPath(*OtherTarget)) );
 				}
 			}
 
-			for (const FFilePath& Path : Target.AdditionalManifestDependencies)
+			for (const FFilePath& Path : Target->Settings.AdditionalManifestDependencies)
 			{
 				ConfigSection.Add( TEXT("ManifestDependencies"), MakePathRelativeToProjectDirectory(Path.FilePath) );
 			}
 
-			const FString SourcePath = ContentDirRelativeToGameDir / TEXT("Localization") / Target.Name;
+			const FString SourcePath = ContentDirRelativeToGameDir / TEXT("Localization") / Target->Settings.Name;
 			ConfigSection.Add( TEXT("SourcePath"), SourcePath );
-			const FString DestinationPath = ContentDirRelativeToGameDir / TEXT("Localization") / Target.Name;
+			const FString DestinationPath = ContentDirRelativeToGameDir / TEXT("Localization") / Target->Settings.Name;
 			ConfigSection.Add( TEXT("DestinationPath"), DestinationPath );
 
 			ConfigSection.Add( TEXT("ManifestName"), FPaths::GetCleanFilename(GetManifestPath(Target)) );
 			ConfigSection.Add( TEXT("ArchiveName"), FPaths::GetCleanFilename(GetArchivePath(Target, FString())) );
 
-			ConfigSection.Add( TEXT("SourceCulture"), Target.NativeCultureStatistics.CultureName );
-			ConfigSection.Add( TEXT("CulturesToGenerate"), Target.NativeCultureStatistics.CultureName );
-			for (const FCultureStatistics& CultureStatistics : Target.SupportedCulturesStatistics)
+			ConfigSection.Add( TEXT("SourceCulture"), Target->Settings.NativeCultureStatistics.CultureName );
+			ConfigSection.Add( TEXT("CulturesToGenerate"), Target->Settings.NativeCultureStatistics.CultureName );
+			for (const FCultureStatistics& CultureStatistics : Target->Settings.SupportedCulturesStatistics)
 			{
 				ConfigSection.Add( TEXT("CulturesToGenerate"), CultureStatistics.CultureName );
 			}
@@ -115,7 +115,7 @@ namespace LocalizationConfigurationScript
 
 		uint32 GatherTextStepIndex = 0;
 		// GatherTextFromSource
-		if (Target.GatherFromTextFiles.SearchDirectories.Num() && Target.GatherFromTextFiles.FileExtensions.Num()) // Don't gather from source if there are no source files.
+		if (Target->Settings.GatherFromTextFiles.SearchDirectories.Num() && Target->Settings.GatherFromTextFiles.FileExtensions.Num()) // Don't gather from source if there are no source files.
 		{
 			FConfigSection& ConfigSection = Script.GatherTextStep(GatherTextStepIndex++);
 
@@ -123,27 +123,27 @@ namespace LocalizationConfigurationScript
 			ConfigSection.Add( TEXT("CommandletClass"), TEXT("GatherTextFromSource") );
 
 			// Include Paths
-			for (const auto& IncludePath : Target.GatherFromTextFiles.SearchDirectories)
+			for (const auto& IncludePath : Target->Settings.GatherFromTextFiles.SearchDirectories)
 			{
 				ConfigSection.Add( TEXT("IncludePaths"), IncludePath );
 			}
 
 			// Exclude Paths
 			ConfigSection.Add( TEXT("ExcludePaths"), ConfigDirRelativeToGameDir / TEXT("Localization") );
-			for (const auto& ExcludePath : Target.GatherFromTextFiles.ExcludePathWildcards)
+			for (const auto& ExcludePath : Target->Settings.GatherFromTextFiles.ExcludePathWildcards)
 			{
 				ConfigSection.Add( TEXT("ExcludePaths"), ExcludePath );
 			}
 
 			// Source File Search Filters
-			for (const auto& FileExtension : Target.GatherFromTextFiles.FileExtensions)
+			for (const auto& FileExtension : Target->Settings.GatherFromTextFiles.FileExtensions)
 			{
 				ConfigSection.Add( TEXT("SourceFileSearchFilters"), FileExtension );
 			}
 		}
 
 		// GatherTextFromAssets
-		if (Target.GatherFromTextFiles.SearchDirectories.Num() && Target.GatherFromTextFiles.FileExtensions.Num()) // Don't gather from packages if there are none specified.
+		if (Target->Settings.GatherFromTextFiles.SearchDirectories.Num() && Target->Settings.GatherFromTextFiles.FileExtensions.Num()) // Don't gather from packages if there are none specified.
 		{
 			FConfigSection& ConfigSection = Script.GatherTextStep(GatherTextStepIndex++);
 
@@ -151,27 +151,27 @@ namespace LocalizationConfigurationScript
 			ConfigSection.Add( TEXT("CommandletClass"), TEXT("GatherTextFromAssets") );
 
 			// Include Paths
-			for (const auto& IncludePath : Target.GatherFromPackages.IncludePathWildcards)
+			for (const auto& IncludePath : Target->Settings.GatherFromPackages.IncludePathWildcards)
 			{
 				ConfigSection.Add( TEXT("IncludePaths"), IncludePath );
 			}
 
 			// Exclude Paths
 			ConfigSection.Add( TEXT("ExcludePaths"), ContentDirRelativeToGameDir / TEXT("Localization") );
-			for (const auto& ExcludePath : Target.GatherFromPackages.ExcludePathWildcards)
+			for (const auto& ExcludePath : Target->Settings.GatherFromPackages.ExcludePathWildcards)
 			{
 				ConfigSection.Add( TEXT("ExcludePaths"), ExcludePath );
 			}
 
 			// Package Extensions
-			for (const auto& FileExtension : Target.GatherFromPackages.FileExtensions)
+			for (const auto& FileExtension : Target->Settings.GatherFromPackages.FileExtensions)
 			{
 				ConfigSection.Add( TEXT("PackageExtensions"), FileExtension );
 			}
 		}
 
 		// GatherTextFromMetadata
-		if (Target.GatherFromMetaData.IncludePathWildcards.Num() && Target.GatherFromMetaData.KeySpecifications.Num()) // Don't gather from metadata if there are none specified.
+		if (Target->Settings.GatherFromMetaData.IncludePathWildcards.Num() && Target->Settings.GatherFromMetaData.KeySpecifications.Num()) // Don't gather from metadata if there are none specified.
 		{
 			FConfigSection& ConfigSection = Script.GatherTextStep(GatherTextStepIndex++);
 
@@ -179,19 +179,19 @@ namespace LocalizationConfigurationScript
 			ConfigSection.Add( TEXT("CommandletClass"), TEXT("GatherTextFromMetadata") );
 
 			// Include Paths
-			for (const auto& IncludePath : Target.GatherFromMetaData.IncludePathWildcards)
+			for (const auto& IncludePath : Target->Settings.GatherFromMetaData.IncludePathWildcards)
 			{
 				ConfigSection.Add( TEXT("IncludePaths"), IncludePath );
 			}
 
 			// Exclude Paths
-			for (const auto& ExcludePath : Target.GatherFromMetaData.ExcludePathWildcards)
+			for (const auto& ExcludePath : Target->Settings.GatherFromMetaData.ExcludePathWildcards)
 			{
 				ConfigSection.Add( TEXT("ExcludePaths"), ExcludePath );
 			}
 
 			// Package Extensions
-			for (const FMetaDataKeyGatherSpecification& Specification : Target.GatherFromMetaData.KeySpecifications)
+			for (const FMetaDataKeyGatherSpecification& Specification : Target->Settings.GatherFromMetaData.KeySpecifications)
 			{
 				ConfigSection.Add( TEXT("InputKeys"), Specification.MetaDataKey );
 				ConfigSection.Add( TEXT("OutputNamespaces"), Specification.TextNamespace );
@@ -220,12 +220,12 @@ namespace LocalizationConfigurationScript
 		return Script;
 	}
 
-	FString GetGatherScriptPath(const FLocalizationTargetSettings& Target)
+	FString GetGatherScriptPath(const ULocalizationTarget* const Target)
 	{
-		return GetScriptDirectory() / FString::Printf( TEXT("%s_Gather.%s"), *(Target.Name), TEXT("ini") );
+		return GetScriptDirectory() / FString::Printf( TEXT("%s_Gather.%s"), *(Target->Settings.Name), TEXT("ini") );
 	}
 
-	FLocalizationConfigurationScript GenerateImportScript(const FLocalizationTargetSettings& Target, const TOptional<FString> CultureName, const TOptional<FString> OutputPathOverride)
+	FLocalizationConfigurationScript GenerateImportScript(const ULocalizationTarget* const Target, const TOptional<FString> CultureName, const TOptional<FString> OutputPathOverride)
 	{
 		FLocalizationConfigurationScript Script;
 
@@ -258,16 +258,16 @@ namespace LocalizationConfigurationScript
 			// Use the default PO file's directory path.
 			else
 			{
-				SourcePath = ContentDirRelativeToGameDir / TEXT("Localization") / Target.Name;
+				SourcePath = ContentDirRelativeToGameDir / TEXT("Localization") / Target->Settings.Name;
 			}
 			ConfigSection.Add( TEXT("SourcePath"), SourcePath );
 
-			const FString DestinationPath = ContentDirRelativeToGameDir / TEXT("Localization") / Target.Name;
+			const FString DestinationPath = ContentDirRelativeToGameDir / TEXT("Localization") / Target->Settings.Name;
 			ConfigSection.Add( TEXT("DestinationPath"), DestinationPath );
 
 			const auto& AddCultureToGenerate = [&](const int32 Index)
 			{
-				ConfigSection.Add( TEXT("CulturesToGenerate"), Target.SupportedCulturesStatistics[Index].CultureName );
+				ConfigSection.Add( TEXT("CulturesToGenerate"), Target->Settings.SupportedCulturesStatistics[Index].CultureName );
 			};
 
 			// Export for a specific culture.
@@ -278,7 +278,7 @@ namespace LocalizationConfigurationScript
 			// Export for all cultures.
 			else
 			{
-				for (const FCultureStatistics& CultureStatistics : Target.SupportedCulturesStatistics)
+				for (const FCultureStatistics& CultureStatistics : Target->Settings.SupportedCulturesStatistics)
 				{
 					ConfigSection.Add( TEXT("CulturesToGenerate"), CultureStatistics.CultureName );
 				}
@@ -312,22 +312,22 @@ namespace LocalizationConfigurationScript
 		return Script;
 	}
 
-	FString GetImportScriptPath(const FLocalizationTargetSettings& Target, const TOptional<FString> CultureName)
+	FString GetImportScriptPath(const ULocalizationTarget* const Target, const TOptional<FString> CultureName)
 	{
 		const FString ConfigFileDirectory = GetScriptDirectory();
 		FString ConfigFilePath;
 		if (CultureName.IsSet())
 		{
-			ConfigFilePath = ConfigFileDirectory / FString::Printf( TEXT("%s_Import_%s.%s"), *Target.Name, *CultureName.GetValue(), TEXT("ini") );
+			ConfigFilePath = ConfigFileDirectory / FString::Printf( TEXT("%s_Import_%s.%s"), *Target->Settings.Name, *CultureName.GetValue(), TEXT("ini") );
 		}
 		else
 		{
-			ConfigFilePath = ConfigFileDirectory / FString::Printf( TEXT("%s_Import.%s"), *Target.Name, TEXT("ini") );
+			ConfigFilePath = ConfigFileDirectory / FString::Printf( TEXT("%s_Import.%s"), *Target->Settings.Name, TEXT("ini") );
 		}
 		return ConfigFilePath;
 	}
 
-	FLocalizationConfigurationScript GenerateExportScript(const FLocalizationTargetSettings& Target, const TOptional<FString> CultureName, const TOptional<FString> OutputPathOverride)
+	FLocalizationConfigurationScript GenerateExportScript(const ULocalizationTarget* const Target, const TOptional<FString> CultureName, const TOptional<FString> OutputPathOverride)
 	{
 		FLocalizationConfigurationScript Script;
 
@@ -342,7 +342,7 @@ namespace LocalizationConfigurationScript
 
 			ConfigSection.Add( TEXT("bExportLoc"), TEXT("true") );
 
-			const FString SourcePath = ContentDirRelativeToGameDir / TEXT("Localization") / Target.Name;
+			const FString SourcePath = ContentDirRelativeToGameDir / TEXT("Localization") / Target->Settings.Name;
 			ConfigSection.Add( TEXT("SourcePath"), SourcePath );
 
 			FString DestinationPath;
@@ -363,13 +363,13 @@ namespace LocalizationConfigurationScript
 			// Use the default PO file's directory path.
 			else
 			{
-				DestinationPath = ContentDirRelativeToGameDir / TEXT("Localization") / Target.Name;
+				DestinationPath = ContentDirRelativeToGameDir / TEXT("Localization") / Target->Settings.Name;
 			}
 			ConfigSection.Add( TEXT("DestinationPath"), DestinationPath );
 
 			TArray<const FCultureStatistics*> AllCultureStatistics;
-			AllCultureStatistics.Add(&Target.NativeCultureStatistics);
-			for (const FCultureStatistics& SupportedCultureStatistics : Target.SupportedCulturesStatistics)
+			AllCultureStatistics.Add(&Target->Settings.NativeCultureStatistics);
+			for (const FCultureStatistics& SupportedCultureStatistics : Target->Settings.SupportedCulturesStatistics)
 			{
 				AllCultureStatistics.Add(&SupportedCultureStatistics);
 			}
@@ -422,22 +422,22 @@ namespace LocalizationConfigurationScript
 		return Script;
 	}
 
-	FString GetExportScriptPath(const FLocalizationTargetSettings& Target, const TOptional<FString> CultureName)
+	FString GetExportScriptPath(const ULocalizationTarget* const Target, const TOptional<FString> CultureName)
 	{
 		const FString ConfigFileDirectory = GetScriptDirectory();
 		FString ConfigFilePath;
 		if (CultureName.IsSet())
 		{
-			ConfigFilePath = ConfigFileDirectory / FString::Printf( TEXT("%s_Export_%s.%s"), *Target.Name, *CultureName.GetValue(), TEXT("ini") );
+			ConfigFilePath = ConfigFileDirectory / FString::Printf( TEXT("%s_Export_%s.%s"), *Target->Settings.Name, *CultureName.GetValue(), TEXT("ini") );
 		}
 		else
 		{
-			ConfigFilePath = ConfigFileDirectory / FString::Printf( TEXT("%s_Export.%s"), *Target.Name, TEXT("ini") );
+			ConfigFilePath = ConfigFileDirectory / FString::Printf( TEXT("%s_Export.%s"), *Target->Settings.Name, TEXT("ini") );
 		}
 		return ConfigFilePath;
 	}
 
-	FLocalizationConfigurationScript GenerateReportScript(const FLocalizationTargetSettings& Target)
+	FLocalizationConfigurationScript GenerateReportScript(const ULocalizationTarget* const Target)
 	{
 		FLocalizationConfigurationScript Script;
 
@@ -450,14 +450,14 @@ namespace LocalizationConfigurationScript
 			// CommandletClass
 			ConfigSection.Add( TEXT("CommandletClass"), TEXT("GenerateTextLocalizationReport") );
 
-			const FString SourcePath = ContentDirRelativeToGameDir / TEXT("Localization") / Target.Name;
+			const FString SourcePath = ContentDirRelativeToGameDir / TEXT("Localization") / Target->Settings.Name;
 			ConfigSection.Add( TEXT("SourcePath"), SourcePath );
-			const FString DestinationPath = ContentDirRelativeToGameDir / TEXT("Localization") / Target.Name;
+			const FString DestinationPath = ContentDirRelativeToGameDir / TEXT("Localization") / Target->Settings.Name;
 			ConfigSection.Add( TEXT("DestinationPath"), DestinationPath );
 
-			ConfigSection.Add( TEXT("ManifestName"), FString::Printf( TEXT("%s.%s"), *Target.Name, TEXT("manifest") ) );
+			ConfigSection.Add( TEXT("ManifestName"), FString::Printf( TEXT("%s.%s"), *Target->Settings.Name, TEXT("manifest") ) );
 
-			for (const FCultureStatistics& CultureStatistics : Target.SupportedCulturesStatistics)
+			for (const FCultureStatistics& CultureStatistics : Target->Settings.SupportedCulturesStatistics)
 			{
 				ConfigSection.Add( TEXT("CulturesToGenerate"), CultureStatistics.CultureName );
 			}
@@ -476,12 +476,12 @@ namespace LocalizationConfigurationScript
 		return Script;
 	}
 
-	FString GetReportScriptPath(const FLocalizationTargetSettings& Target)
+	FString GetReportScriptPath(const ULocalizationTarget* const Target)
 	{
-		return GetScriptDirectory() / FString::Printf( TEXT("%s_GenerateReports.%s"), *(Target.Name), TEXT("ini") );
+		return GetScriptDirectory() / FString::Printf( TEXT("%s_GenerateReports.%s"), *(Target->Settings.Name), TEXT("ini") );
 	}
 
-	FLocalizationConfigurationScript GenerateCompileScript(const FLocalizationTargetSettings& Target)
+	FLocalizationConfigurationScript GenerateCompileScript(const ULocalizationTarget* const Target)
 	{
 		FLocalizationConfigurationScript Script;
 
@@ -494,17 +494,17 @@ namespace LocalizationConfigurationScript
 			// CommandletClass
 			ConfigSection.Add( TEXT("CommandletClass"), TEXT("GenerateTextLocalizationResource") );
 
-			const FString SourcePath = ContentDirRelativeToGameDir / TEXT("Localization") / Target.Name;
+			const FString SourcePath = ContentDirRelativeToGameDir / TEXT("Localization") / Target->Settings.Name;
 			ConfigSection.Add( TEXT("SourcePath"), SourcePath );
-			const FString DestinationPath = ContentDirRelativeToGameDir / TEXT("Localization") / Target.Name;
+			const FString DestinationPath = ContentDirRelativeToGameDir / TEXT("Localization") / Target->Settings.Name;
 			ConfigSection.Add( TEXT("DestinationPath"), DestinationPath );
 
-			ConfigSection.Add( TEXT("ManifestName"), FString::Printf( TEXT("%s.%s"), *Target.Name, TEXT("manifest") ) );
-			ConfigSection.Add( TEXT("ResourceName"), FString::Printf( TEXT("%s.%s"), *Target.Name, TEXT("locres") ) );
+			ConfigSection.Add( TEXT("ManifestName"), FString::Printf( TEXT("%s.%s"), *Target->Settings.Name, TEXT("manifest") ) );
+			ConfigSection.Add( TEXT("ResourceName"), FString::Printf( TEXT("%s.%s"), *Target->Settings.Name, TEXT("locres") ) );
 
 			TArray<const FCultureStatistics*> AllCultureStatistics;
-			AllCultureStatistics.Add(&Target.NativeCultureStatistics);
-			for (const FCultureStatistics& SupportedCultureStatistics : Target.SupportedCulturesStatistics)
+			AllCultureStatistics.Add(&Target->Settings.NativeCultureStatistics);
+			for (const FCultureStatistics& SupportedCultureStatistics : Target->Settings.SupportedCulturesStatistics)
 			{
 				AllCultureStatistics.Add(&SupportedCultureStatistics);
 			}
@@ -520,8 +520,8 @@ namespace LocalizationConfigurationScript
 		return Script;
 	}
 
-	FString GetCompileScriptPath(const FLocalizationTargetSettings& Target)
+	FString GetCompileScriptPath(const ULocalizationTarget* const Target)
 	{
-		return GetScriptDirectory() / FString::Printf( TEXT("%s_Compile.%s"), *(Target.Name), TEXT("ini") );
+		return GetScriptDirectory() / FString::Printf( TEXT("%s_Compile.%s"), *(Target->Settings.Name), TEXT("ini") );
 	}
 }
\ No newline at end of file
diff --git a/Engine/Source/Editor/LocalizationDashboard/Private/LocalizationConfigurationScript.h b/Engine/Source/Editor/LocalizationDashboard/Private/LocalizationConfigurationScript.h
--- a/Engine/Source/Editor/LocalizationDashboard/Private/LocalizationConfigurationScript.h
+++ b/Engine/Source/Editor/LocalizationDashboard/Private/LocalizationConfigurationScript.h
@@ -22,30 +22,30 @@ namespace LocalizationConfigurationScript
 {
 	FString MakePathRelativeToProjectDirectory(const FString& Path);
 
-	FString GetScriptDirectory();
-	FString GetDataDirectory(const FLocalizationTargetSettings& Target);
-	TArray<FString> GetScriptPaths(const FLocalizationTargetSettings& Target);
+	FString GetScriptDirectory(const ULocalizationTarget* const Target);
+	FString GetDataDirectory(const ULocalizationTarget* const Target);
+	TArray<FString> GetScriptPaths(const ULocalizationTarget* const Target);
 
-	FString GetManifestPath(const FLocalizationTargetSettings& Target);
-	FString GetArchivePath(const FLocalizationTargetSettings& Target, const FString& CultureName);
-	FString GetDefaultPOFileName(const FLocalizationTargetSettings& Target);
-	FString GetDefaultPOPath(const FLocalizationTargetSettings& Target, const FString& CultureName);
-	FString GetLocResPath(const FLocalizationTargetSettings& Target, const FString& CultureName);
-	FString GetWordCountCSVPath(const FLocalizationTargetSettings& Target);
-	FString GetConflictReportPath(const FLocalizationTargetSettings& Target);
+	FString GetManifestPath(const ULocalizationTarget* const Target);
+	FString GetArchivePath(const ULocalizationTarget* const Target, const FString& CultureName);
+	FString GetDefaultPOFileName(const ULocalizationTarget* const Target);
+	FString GetDefaultPOPath(const ULocalizationTarget* const Target, const FString& CultureName);
+	FString GetLocResPath(const ULocalizationTarget* const Target, const FString& CultureName);
+	FString GetWordCountCSVPath(const ULocalizationTarget* const Target);
+	FString GetConflictReportPath(const ULocalizationTarget* const Target);
 
-	FLocalizationConfigurationScript GenerateGatherScript(const FLocalizationTargetSettings& Target);
-	FString GetGatherScriptPath(const FLocalizationTargetSettings& Target);
+	FLocalizationConfigurationScript GenerateGatherScript(const ULocalizationTarget* const Target);
+	FString GetGatherScriptPath(const ULocalizationTarget* const Target);
 
-	FLocalizationConfigurationScript GenerateImportScript(const FLocalizationTargetSettings& Target, const TOptional<FString> CultureName = TOptional<FString>(), const TOptional<FString> OutputPathOverride = TOptional<FString>());
-	FString GetImportScriptPath(const FLocalizationTargetSettings& Target, const TOptional<FString> CultureName = TOptional<FString>());
+	FLocalizationConfigurationScript GenerateImportScript(const ULocalizationTarget* const Target, const TOptional<FString> CultureName = TOptional<FString>(), const TOptional<FString> OutputPathOverride = TOptional<FString>());
+	FString GetImportScriptPath(const ULocalizationTarget* const Target, const TOptional<FString> CultureName = TOptional<FString>());
 
-	FLocalizationConfigurationScript GenerateExportScript(const FLocalizationTargetSettings& Target, const TOptional<FString> CultureName = TOptional<FString>(), const TOptional<FString> OutputPathOverride = TOptional<FString>());
-	FString GetExportScriptPath(const FLocalizationTargetSettings& Target, const TOptional<FString> CultureName = TOptional<FString>());
+	FLocalizationConfigurationScript GenerateExportScript(const ULocalizationTarget* const Target, const TOptional<FString> CultureName = TOptional<FString>(), const TOptional<FString> OutputPathOverride = TOptional<FString>());
+	FString GetExportScriptPath(const ULocalizationTarget* const Target, const TOptional<FString> CultureName = TOptional<FString>());
 
-	FLocalizationConfigurationScript GenerateReportScript(const FLocalizationTargetSettings& Target);
-	FString GetReportScriptPath(const FLocalizationTargetSettings& Target);
+	FLocalizationConfigurationScript GenerateReportScript(const ULocalizationTarget* const Target);
+	FString GetReportScriptPath(const ULocalizationTarget* const Target);
 
-	FLocalizationConfigurationScript GenerateCompileScript(const FLocalizationTargetSettings& Target);
-	FString GetCompileScriptPath(const FLocalizationTargetSettings& Target);
+	FLocalizationConfigurationScript GenerateCompileScript(const ULocalizationTarget* const Target);
+	FString GetCompileScriptPath(const ULocalizationTarget* const Target);
 }
\ No newline at end of file
diff --git a/Engine/Source/Editor/LocalizationDashboard/Private/LocalizationTarget.cpp b/Engine/Source/Editor/LocalizationDashboard/Private/LocalizationTarget.cpp
--- a/Engine/Source/Editor/LocalizationDashboard/Private/LocalizationTarget.cpp
+++ b/Engine/Source/Editor/LocalizationDashboard/Private/LocalizationTarget.cpp
@@ -170,13 +170,13 @@ void ULocalizationTarget::PostEditChangeProperty(FPropertyChangedEvent& Property
 
 bool ULocalizationTarget::UpdateWordCountsFromCSV()
 {
-	const FString CSVFilePath = LocalizationConfigurationScript::GetWordCountCSVPath(Settings);
+	const FString CSVFilePath = LocalizationConfigurationScript::GetWordCountCSVPath(this);
 	return FWordCountCSVParser::Execute(Settings, CSVFilePath);
 }
 
 void ULocalizationTarget::UpdateStatusFromConflictReport()
 {
-	const FString ConflictReportFilePath = LocalizationConfigurationScript::GetConflictReportPath(Settings);
+	const FString ConflictReportFilePath = LocalizationConfigurationScript::GetConflictReportPath(this);
 	const int64 ConflictReportFileSize = IFileManager::Get().FileSize(*ConflictReportFilePath);
 
 	switch (ConflictReportFileSize)
@@ -243,7 +243,7 @@ bool ULocalizationTarget::RenameTargetAndFiles(const FString& NewName)
 	};
 
 	// Delete configuration files.
-	const TArray<FString>& ScriptPaths = LocalizationConfigurationScript::GetScriptPaths(Settings);
+	const TArray<FString>& ScriptPaths = LocalizationConfigurationScript::GetScriptPaths(this);
 	for (const FString& ScriptPath : ScriptPaths)
 	{
 		if (!TryDelete(ScriptPath))
@@ -254,9 +254,9 @@ bool ULocalizationTarget::RenameTargetAndFiles(const FString& NewName)
 
 	const auto& GetNamedPaths = [this](TArray<FString>& NamedPaths)
 	{
-		NamedPaths.Add(LocalizationConfigurationScript::GetManifestPath(Settings));
-		NamedPaths.Add(LocalizationConfigurationScript::GetWordCountCSVPath(Settings));
-		NamedPaths.Add(LocalizationConfigurationScript::GetConflictReportPath(Settings));
+		NamedPaths.Add(LocalizationConfigurationScript::GetManifestPath(this));
+		NamedPaths.Add(LocalizationConfigurationScript::GetWordCountCSVPath(this));
+		NamedPaths.Add(LocalizationConfigurationScript::GetConflictReportPath(this));
 
 		TArray<const FCultureStatistics*> Cultures;
 		Cultures.Add(&Settings.NativeCultureStatistics);
@@ -267,12 +267,12 @@ bool ULocalizationTarget::RenameTargetAndFiles(const FString& NewName)
 
 		for (const FCultureStatistics* Culture : Cultures)
 		{
-			NamedPaths.Add(LocalizationConfigurationScript::GetArchivePath(Settings, Culture->CultureName));
-			NamedPaths.Add(LocalizationConfigurationScript::GetDefaultPOPath(Settings, Culture->CultureName));
-			NamedPaths.Add(LocalizationConfigurationScript::GetLocResPath(Settings, Culture->CultureName));
+			NamedPaths.Add(LocalizationConfigurationScript::GetArchivePath(this, Culture->CultureName));
+			NamedPaths.Add(LocalizationConfigurationScript::GetDefaultPOPath(this, Culture->CultureName));
+			NamedPaths.Add(LocalizationConfigurationScript::GetLocResPath(this, Culture->CultureName));
 		}
 
-		NamedPaths.Add(LocalizationConfigurationScript::GetDataDirectory(Settings));
+		NamedPaths.Add(LocalizationConfigurationScript::GetDataDirectory(this));
 	};
 
 	TArray<FString> OldPaths;
@@ -333,17 +333,17 @@ bool ULocalizationTarget::RenameTargetAndFiles(const FString& NewName)
 	}
 
 	// Generate new configuration files.
-	LocalizationConfigurationScript::GenerateGatherScript(Settings).Write(LocalizationConfigurationScript::GetGatherScriptPath(Settings));
-	LocalizationConfigurationScript::GenerateImportScript(Settings).Write(LocalizationConfigurationScript::GetImportScriptPath(Settings));
-	LocalizationConfigurationScript::GenerateExportScript(Settings).Write(LocalizationConfigurationScript::GetExportScriptPath(Settings));
-	LocalizationConfigurationScript::GenerateReportScript(Settings).Write(LocalizationConfigurationScript::GetReportScriptPath(Settings));
+	LocalizationConfigurationScript::GenerateGatherScript(this).Write(LocalizationConfigurationScript::GetGatherScriptPath(this));
+	LocalizationConfigurationScript::GenerateImportScript(this).Write(LocalizationConfigurationScript::GetImportScriptPath(this));
+	LocalizationConfigurationScript::GenerateExportScript(this).Write(LocalizationConfigurationScript::GetExportScriptPath(this));
+	LocalizationConfigurationScript::GenerateReportScript(this).Write(LocalizationConfigurationScript::GetReportScriptPath(this));
 
 	//if (CanUseSourceControl)
 	//{
-	//	HasCompletelySucceeded = HasCompletelySucceeded && SourceControlProvider.Execute(ISourceControlOperation::Create<FMarkForAdd>(), *LocalizationConfigurationScript::GetGatherScriptPath(Settings)) == ECommandResult::Succeeded;
-	//	HasCompletelySucceeded = HasCompletelySucceeded && SourceControlProvider.Execute(ISourceControlOperation::Create<FMarkForAdd>(), *LocalizationConfigurationScript::GetImportScriptPath(Settings)) == ECommandResult::Succeeded;
-	//	HasCompletelySucceeded = HasCompletelySucceeded && SourceControlProvider.Execute(ISourceControlOperation::Create<FMarkForAdd>(), *LocalizationConfigurationScript::GetExportScriptPath(Settings)) == ECommandResult::Succeeded;
-	//	HasCompletelySucceeded = HasCompletelySucceeded && SourceControlProvider.Execute(ISourceControlOperation::Create<FMarkForAdd>(), *LocalizationConfigurationScript::GetReportScriptPath(Settings)) == ECommandResult::Succeeded;
+	//	HasCompletelySucceeded = HasCompletelySucceeded && SourceControlProvider.Execute(ISourceControlOperation::Create<FMarkForAdd>(), *LocalizationConfigurationScript::GetGatherScriptPath(this)) == ECommandResult::Succeeded;
+	//	HasCompletelySucceeded = HasCompletelySucceeded && SourceControlProvider.Execute(ISourceControlOperation::Create<FMarkForAdd>(), *LocalizationConfigurationScript::GetImportScriptPath(this)) == ECommandResult::Succeeded;
+	//	HasCompletelySucceeded = HasCompletelySucceeded && SourceControlProvider.Execute(ISourceControlOperation::Create<FMarkForAdd>(), *LocalizationConfigurationScript::GetExportScriptPath(this)) == ECommandResult::Succeeded;
+	//	HasCompletelySucceeded = HasCompletelySucceeded && SourceControlProvider.Execute(ISourceControlOperation::Create<FMarkForAdd>(), *LocalizationConfigurationScript::GetReportScriptPath(this)) == ECommandResult::Succeeded;
 	//}
 
 	return HasCompletelySucceeded;
@@ -355,7 +355,7 @@ bool ULocalizationTarget::DeleteFiles(const FString* const Culture) const
 	ISourceControlModule& SourceControl = ISourceControlModule::Get();
 	ISourceControlProvider& SourceControlProvider = SourceControl.GetProvider();
 
-	const FString DataDirectory = LocalizationConfigurationScript::GetDataDirectory(Settings) + (Culture ? *Culture : FString());
+	const FString DataDirectory = LocalizationConfigurationScript::GetDataDirectory(this) + (Culture ? *Culture : FString());
 
 	// Remove data files from source control.
 	{
@@ -399,7 +399,7 @@ bool ULocalizationTarget::DeleteFiles(const FString* const Culture) const
 	// Delete configuration files if deleting the target.
 	if (!Culture)
 	{
-		const TArray<FString>& ScriptPaths = LocalizationConfigurationScript::GetScriptPaths(Settings);
+		const TArray<FString>& ScriptPaths = LocalizationConfigurationScript::GetScriptPaths(this);
 
 		// Remove script files from source control.
 		{
diff --git a/Engine/Source/Editor/LocalizationDashboard/Private/LocalizationTargetDetailCustomization.cpp b/Engine/Source/Editor/LocalizationDashboard/Private/LocalizationTargetDetailCustomization.cpp
--- a/Engine/Source/Editor/LocalizationDashboard/Private/LocalizationTargetDetailCustomization.cpp
+++ b/Engine/Source/Editor/LocalizationDashboard/Private/LocalizationTargetDetailCustomization.cpp
@@ -614,7 +614,7 @@ void FLocalizationTargetDetailCustomization::Gather()
 
 		// Execute gather.
 		const TSharedPtr<SWindow> ParentWindow = FSlateApplication::Get().FindWidgetWindow(DetailLayoutBuilder->GetDetailsView().AsShared());
-		LocalizationCommandletTasks::GatherTarget(ParentWindow.ToSharedRef(), LocalizationTarget->Settings);
+		LocalizationCommandletTasks::GatherTarget(ParentWindow.ToSharedRef(), LocalizationTarget.Get());
 
 		UpdateTargetFromReports();
 	}
@@ -632,7 +632,7 @@ void FLocalizationTargetDetailCustomization::ImportAllCultures()
 			ParentWindowWindowHandle = ParentWindow->GetNativeWindow()->GetOSWindowHandle();
 		}
 
-		const FString DefaultPath = FPaths::ConvertRelativePathToFull(LocalizationConfigurationScript::GetDataDirectory(LocalizationTarget->Settings));
+		const FString DefaultPath = FPaths::ConvertRelativePathToFull(LocalizationConfigurationScript::GetDataDirectory(LocalizationTarget.Get()));
 
 		FText DialogTitle;
 		{
@@ -645,7 +645,7 @@ void FLocalizationTargetDetailCustomization::ImportAllCultures()
 		FString OutputDirectory;
 		if (DesktopPlatform->OpenDirectoryDialog(ParentWindowWindowHandle, DialogTitle.ToString(), DefaultPath, OutputDirectory))
 		{
-			LocalizationCommandletTasks::ImportTarget(ParentWindow.ToSharedRef(), LocalizationTarget->Settings, TOptional<FString>(OutputDirectory));
+			LocalizationCommandletTasks::ImportTarget(ParentWindow.ToSharedRef(), LocalizationTarget.Get(), TOptional<FString>(OutputDirectory));
 
 			UpdateTargetFromReports();
 		}
@@ -664,7 +664,7 @@ void FLocalizationTargetDetailCustomization::ExportAllCultures()
 			ParentWindowWindowHandle = ParentWindow->GetNativeWindow()->GetOSWindowHandle();
 		}
 
-		const FString DefaultPath = FPaths::ConvertRelativePathToFull(LocalizationConfigurationScript::GetDataDirectory(LocalizationTarget->Settings));
+		const FString DefaultPath = FPaths::ConvertRelativePathToFull(LocalizationConfigurationScript::GetDataDirectory(LocalizationTarget.Get()));
 
 		FText DialogTitle;
 		{
@@ -677,7 +677,7 @@ void FLocalizationTargetDetailCustomization::ExportAllCultures()
 		FString OutputDirectory;
 		if (DesktopPlatform->OpenDirectoryDialog(ParentWindowWindowHandle, DialogTitle.ToString(), DefaultPath, OutputDirectory))
 		{
-			LocalizationCommandletTasks::ExportTarget(ParentWindow.ToSharedRef(), LocalizationTarget->Settings, TOptional<FString>(OutputDirectory));
+			LocalizationCommandletTasks::ExportTarget(ParentWindow.ToSharedRef(), LocalizationTarget.Get(), TOptional<FString>(OutputDirectory));
 		}
 	}
 }
@@ -688,7 +688,7 @@ void FLocalizationTargetDetailCustomization::Compile()
 	{
 		// Execute compile.
 		const TSharedPtr<SWindow> ParentWindow = FSlateApplication::Get().FindWidgetWindow(DetailLayoutBuilder->GetDetailsView().AsShared());
-		LocalizationCommandletTasks::CompileTarget(ParentWindow.ToSharedRef(), LocalizationTarget->Settings);
+		LocalizationCommandletTasks::CompileTarget(ParentWindow.ToSharedRef(), LocalizationTarget.Get());
 	}
 }
 
@@ -697,7 +697,7 @@ void FLocalizationTargetDetailCustomization::RefreshWordCounts()
 	if (LocalizationTarget.IsValid())
 	{
 		const TSharedPtr<SWindow> ParentWindow = FSlateApplication::Get().FindWidgetWindow(DetailLayoutBuilder->GetDetailsView().AsShared());
-		LocalizationCommandletTasks::GenerateReportsForTarget(ParentWindow.ToSharedRef(), LocalizationTarget->Settings);
+		LocalizationCommandletTasks::GenerateReportsForTarget(ParentWindow.ToSharedRef(), LocalizationTarget.Get());
 
 		UpdateTargetFromReports();
 	}
diff --git a/Engine/Source/Editor/LocalizationDashboard/Private/LocalizationTargetSetDetailCustomization.cpp b/Engine/Source/Editor/LocalizationDashboard/Private/LocalizationTargetSetDetailCustomization.cpp
--- a/Engine/Source/Editor/LocalizationDashboard/Private/LocalizationTargetSetDetailCustomization.cpp
+++ b/Engine/Source/Editor/LocalizationDashboard/Private/LocalizationTargetSetDetailCustomization.cpp
@@ -287,12 +287,7 @@ void FLocalizationTargetSetDetailCustomization::GatherAllTargets()
 
 	// Execute gather.
 	const TSharedPtr<SWindow> ParentWindow = FSlateApplication::Get().FindWidgetWindow(DetailLayoutBuilder->GetDetailsView().AsShared());
-	TArray<FLocalizationTargetSettings*> TargetSettings;
-	for (ULocalizationTarget* const TargetObject : TargetSet->TargetObjects)
-	{
-		TargetSettings.Add(&TargetObject->Settings);
-	}
-	LocalizationCommandletTasks::GatherTargets(ParentWindow.ToSharedRef(), TargetSettings);
+	LocalizationCommandletTasks::GatherTargets(ParentWindow.ToSharedRef(), TargetSet->TargetObjects);
 
 	for (ULocalizationTarget* const LocalizationTarget : TargetSet->TargetObjects)
 	{
@@ -319,12 +314,7 @@ void FLocalizationTargetSetDetailCustomization::ImportAllTargets()
 		FString OutputDirectory;
 		if (DesktopPlatform->OpenDirectoryDialog(ParentWindowWindowHandle, LOCTEXT("ImportAllTranslationsDialogTitle", "Import All Translations from Directory").ToString(), DefaultPath, OutputDirectory))
 		{
-			TArray<FLocalizationTargetSettings*> TargetSettings;
-			for (ULocalizationTarget* const TargetObject : TargetSet->TargetObjects)
-			{
-				TargetSettings.Add(&TargetObject->Settings);
-			}
-			LocalizationCommandletTasks::ImportTargets(ParentWindow.ToSharedRef(), TargetSettings, TOptional<FString>(OutputDirectory));
+			LocalizationCommandletTasks::ImportTargets(ParentWindow.ToSharedRef(), TargetSet->TargetObjects, TOptional<FString>(OutputDirectory));
 
 			for (ULocalizationTarget* const LocalizationTarget : TargetSet->TargetObjects)
 			{
@@ -352,12 +342,7 @@ void FLocalizationTargetSetDetailCustomization::ExportAllTargets()
 		FString OutputDirectory;
 		if (DesktopPlatform->OpenDirectoryDialog(ParentWindowWindowHandle, LOCTEXT("ExportAllTranslationsDialogTitle", "Export All Translations to Directory").ToString(), DefaultPath, OutputDirectory))
 		{
-			TArray<FLocalizationTargetSettings*> TargetSettings;
-			for (ULocalizationTarget* const TargetObject : TargetSet->TargetObjects)
-			{
-				TargetSettings.Add(&TargetObject->Settings);
-			}
-			LocalizationCommandletTasks::ExportTargets(ParentWindow.ToSharedRef(), TargetSettings, TOptional<FString>(OutputDirectory));
+			LocalizationCommandletTasks::ExportTargets(ParentWindow.ToSharedRef(), TargetSet->TargetObjects, TOptional<FString>(OutputDirectory));
 		}
 	}
 }
@@ -415,12 +400,7 @@ void FLocalizationTargetSetDetailCustomization::CompileAllTargets()
 {
 	// Execute compile.
 	const TSharedPtr<SWindow> ParentWindow = FSlateApplication::Get().FindWidgetWindow(DetailLayoutBuilder->GetDetailsView().AsShared());
-	TArray<FLocalizationTargetSettings*> TargetSettings;
-	for (ULocalizationTarget* const TargetObject : TargetSet->TargetObjects)
-	{
-		TargetSettings.Add(&TargetObject->Settings);
-	}
-	LocalizationCommandletTasks::CompileTargets(ParentWindow.ToSharedRef(), TargetSettings);
+	LocalizationCommandletTasks::CompileTargets(ParentWindow.ToSharedRef(), TargetSet->TargetObjects);
 }
 
 TSharedRef<ITableRow> FLocalizationTargetSetDetailCustomization::OnGenerateRow(TSharedPtr<IPropertyHandle> TargetObjectPropertyHandle, const TSharedRef<STableViewBase>& Table)
diff --git a/Engine/Source/Editor/LocalizationDashboard/Private/SLocalizationDashboardTargetRow.cpp b/Engine/Source/Editor/LocalizationDashboard/Private/SLocalizationDashboardTargetRow.cpp
--- a/Engine/Source/Editor/LocalizationDashboard/Private/SLocalizationDashboardTargetRow.cpp
+++ b/Engine/Source/Editor/LocalizationDashboard/Private/SLocalizationDashboardTargetRow.cpp
@@ -314,7 +314,7 @@ FReply SLocalizationDashboardTargetRow::RefreshWordCount()
 	if (LocalizationTarget)
 	{
 		const TSharedPtr<SWindow> ParentWindow = FSlateApplication::Get().FindWidgetWindow(AsShared());
-		LocalizationCommandletTasks::GenerateReportsForTarget(ParentWindow.ToSharedRef(), LocalizationTarget->Settings);
+		LocalizationCommandletTasks::GenerateReportsForTarget(ParentWindow.ToSharedRef(), LocalizationTarget);
 
 		UpdateTargetFromReports();
 	}
@@ -354,7 +354,7 @@ FReply SLocalizationDashboardTargetRow::Gather()
 
 		// Execute gather.
 		const TSharedPtr<SWindow> ParentWindow = FSlateApplication::Get().FindWidgetWindow(AsShared());
-		LocalizationCommandletTasks::GatherTarget(ParentWindow.ToSharedRef(), LocalizationTarget->Settings);
+		LocalizationCommandletTasks::GatherTarget(ParentWindow.ToSharedRef(), LocalizationTarget);
 
 		UpdateTargetFromReports();
 	}
@@ -375,7 +375,7 @@ FReply SLocalizationDashboardTargetRow::ImportAll()
 			ParentWindowWindowHandle = ParentWindow->GetNativeWindow()->GetOSWindowHandle();
 		}
 
-		const FString DefaultPath = FPaths::ConvertRelativePathToFull(LocalizationConfigurationScript::GetDataDirectory(LocalizationTarget->Settings));
+		const FString DefaultPath = FPaths::ConvertRelativePathToFull(LocalizationConfigurationScript::GetDataDirectory(LocalizationTarget));
 
 		FText DialogTitle;
 		{
@@ -388,7 +388,7 @@ FReply SLocalizationDashboardTargetRow::ImportAll()
 		FString OutputDirectory;
 		if (DesktopPlatform->OpenDirectoryDialog(ParentWindowWindowHandle, DialogTitle.ToString(), DefaultPath, OutputDirectory))
 		{
-			LocalizationCommandletTasks::ImportTarget(ParentWindow.ToSharedRef(), LocalizationTarget->Settings, TOptional<FString>(OutputDirectory));
+			LocalizationCommandletTasks::ImportTarget(ParentWindow.ToSharedRef(), LocalizationTarget, TOptional<FString>(OutputDirectory));
 			UpdateTargetFromReports();
 		}
 	}
@@ -409,7 +409,7 @@ FReply SLocalizationDashboardTargetRow::ExportAll()
 			ParentWindowWindowHandle = ParentWindow->GetNativeWindow()->GetOSWindowHandle();
 		}
 
-		const FString DefaultPath = FPaths::ConvertRelativePathToFull(LocalizationConfigurationScript::GetDataDirectory(LocalizationTarget->Settings));
+		const FString DefaultPath = FPaths::ConvertRelativePathToFull(LocalizationConfigurationScript::GetDataDirectory(LocalizationTarget));
 
 		FText DialogTitle;
 		{
@@ -422,7 +422,7 @@ FReply SLocalizationDashboardTargetRow::ExportAll()
 		FString OutputDirectory;
 		if (DesktopPlatform->OpenDirectoryDialog(ParentWindowWindowHandle, DialogTitle.ToString(), DefaultPath, OutputDirectory))
 		{
-			LocalizationCommandletTasks::ExportTarget(ParentWindow.ToSharedRef(), LocalizationTarget->Settings, TOptional<FString>(OutputDirectory));
+			LocalizationCommandletTasks::ExportTarget(ParentWindow.ToSharedRef(), LocalizationTarget, TOptional<FString>(OutputDirectory));
 		}
 	}
 
@@ -436,7 +436,7 @@ FReply SLocalizationDashboardTargetRow::Compile()
 	{
 		// Execute compile.
 		const TSharedPtr<SWindow> ParentWindow = FSlateApplication::Get().FindWidgetWindow(AsShared());
-		LocalizationCommandletTasks::CompileTarget(ParentWindow.ToSharedRef(), LocalizationTarget->Settings);
+		LocalizationCommandletTasks::CompileTarget(ParentWindow.ToSharedRef(), LocalizationTarget);
 	}
 
 	return FReply::Handled();
diff --git a/Engine/Source/Editor/LocalizationDashboard/Private/SLocalizationTargetEditorCultureRow.cpp b/Engine/Source/Editor/LocalizationDashboard/Private/SLocalizationTargetEditorCultureRow.cpp
--- a/Engine/Source/Editor/LocalizationDashboard/Private/SLocalizationTargetEditorCultureRow.cpp
+++ b/Engine/Source/Editor/LocalizationDashboard/Private/SLocalizationTargetEditorCultureRow.cpp
@@ -301,11 +301,10 @@ void SLocalizationTargetEditorCultureRow::UpdateTargetFromReports()
 FReply SLocalizationTargetEditorCultureRow::Edit()
 {
 	const FCulturePtr Culture = GetCulture();
-
-	FLocalizationTargetSettings* const TargetSettings = GetTargetSettings();
-	if (TargetSettings && Culture.IsValid())
+	ULocalizationTarget* const LocalizationTarget = GetTarget();
+	if (Culture.IsValid() && LocalizationTarget)
 	{
-		FModuleManager::LoadModuleChecked<ITranslationEditor>("TranslationEditor").OpenTranslationEditor(LocalizationConfigurationScript::GetManifestPath(*TargetSettings), LocalizationConfigurationScript::GetArchivePath(*TargetSettings, Culture->GetName()));
+		FModuleManager::LoadModuleChecked<ITranslationEditor>("TranslationEditor").OpenTranslationEditor(LocalizationConfigurationScript::GetManifestPath(LocalizationTarget), LocalizationConfigurationScript::GetArchivePath(LocalizationTarget, Culture->GetName()));
 	}
 
 	return FReply::Handled();
@@ -314,9 +313,9 @@ FReply SLocalizationTargetEditorCultureRow::Edit()
 FReply SLocalizationTargetEditorCultureRow::Import()
 {
 	const FCulturePtr Culture = GetCulture();
-	FLocalizationTargetSettings* const TargetSettings = GetTargetSettings();
+	ULocalizationTarget* const LocalizationTarget = GetTarget();
 	IDesktopPlatform* DesktopPlatform = FDesktopPlatformModule::Get();
-	if (Culture.IsValid() && TargetSettings && DesktopPlatform)
+	if (Culture.IsValid() && LocalizationTarget && DesktopPlatform)
 	{
 		void* ParentWindowWindowHandle = NULL;
 		{
@@ -327,18 +326,18 @@ FReply SLocalizationTargetEditorCultureRow::Import()
 			}
 		}
 
-		const FString POFileName = LocalizationConfigurationScript::GetDefaultPOFileName(*TargetSettings);
+		const FString POFileName = LocalizationConfigurationScript::GetDefaultPOFileName(LocalizationTarget);
 		const FString POFileTypeDescription = LOCTEXT("PortableObjectFileDescription", "Portable Object").ToString();
 		const FString POFileExtension = FPaths::GetExtension(POFileName);
 		const FString POFileExtensionWildcard = FString::Printf(TEXT("*.%s"), *POFileExtension);
 		const FString FileTypes = FString::Printf(TEXT("%s (%s)|%s"), *POFileTypeDescription, *POFileExtensionWildcard, *POFileExtensionWildcard);
 		const FString DefaultFilename = POFileName;
-		const FString DefaultPath = FPaths::GetPath(LocalizationConfigurationScript::GetDefaultPOPath(*TargetSettings, Culture->GetName()));
+		const FString DefaultPath = FPaths::GetPath(LocalizationConfigurationScript::GetDefaultPOPath(LocalizationTarget, Culture->GetName()));
 
 		FText DialogTitle;
 		{
 			FFormatNamedArguments FormatArguments;
-			FormatArguments.Add(TEXT("TargetName"), FText::FromString(TargetSettings->Name));
+			FormatArguments.Add(TEXT("TargetName"), FText::FromString(LocalizationTarget->Settings.Name));
 			FormatArguments.Add(TEXT("CultureName"), FText::FromString(Culture->GetDisplayName()));
 			DialogTitle = FText::Format(LOCTEXT("ImportSpecificTranslationsForTargetDialogTitleFormat", "Import {CultureName} Translations for {TargetName} from Directory"), FormatArguments);
 		}
@@ -348,7 +347,7 @@ FReply SLocalizationTargetEditorCultureRow::Import()
 		if (DesktopPlatform->OpenFileDialog(ParentWindowWindowHandle, DialogTitle.ToString(), DefaultPath, DefaultFilename, FileTypes, 0, OpenFilenames))
 		{
 			const TSharedPtr<SWindow> ParentWindow = FSlateApplication::Get().FindWidgetWindow(AsShared());
-			LocalizationCommandletTasks::ImportCulture(ParentWindow.ToSharedRef(), *TargetSettings, Culture->GetName(), TOptional<FString>(OpenFilenames.Top()));
+			LocalizationCommandletTasks::ImportCulture(ParentWindow.ToSharedRef(), LocalizationTarget, Culture->GetName(), TOptional<FString>(OpenFilenames.Top()));
 
 			UpdateTargetFromReports();
 		}
@@ -360,9 +359,9 @@ FReply SLocalizationTargetEditorCultureRow::Import()
 FReply SLocalizationTargetEditorCultureRow::Export()
 {
 	const FCulturePtr Culture = GetCulture();
-	FLocalizationTargetSettings* const TargetSettings = GetTargetSettings();
+	ULocalizationTarget* const LocalizationTarget = GetTarget();
 	IDesktopPlatform* DesktopPlatform = FDesktopPlatformModule::Get();
-	if (Culture.IsValid() && TargetSettings && DesktopPlatform)
+	if (Culture.IsValid() && LocalizationTarget && DesktopPlatform)
 	{
 		void* ParentWindowWindowHandle = NULL;
 		const TSharedPtr<SWindow> ParentWindow = FSlateApplication::Get().FindWidgetWindow(AsShared());
@@ -371,18 +370,18 @@ FReply SLocalizationTargetEditorCultureRow::Export()
 			ParentWindowWindowHandle = ParentWindow->GetNativeWindow()->GetOSWindowHandle();
 		}
 
-		const FString POFileName = LocalizationConfigurationScript::GetDefaultPOFileName(*TargetSettings);
+		const FString POFileName = LocalizationConfigurationScript::GetDefaultPOFileName(LocalizationTarget);
 		const FString POFileTypeDescription = LOCTEXT("PortableObjectFileDescription", "Portable Object").ToString();
 		const FString POFileExtension = FPaths::GetExtension(POFileName);
 		const FString POFileExtensionWildcard = FString::Printf(TEXT("*.%s"), *POFileExtension);
 		const FString FileTypes = FString::Printf(TEXT("%s (%s)|%s"), *POFileTypeDescription, *POFileExtensionWildcard, *POFileExtensionWildcard);
 		const FString DefaultFilename = POFileName;
-		const FString DefaultPath = FPaths::GetPath(LocalizationConfigurationScript::GetDefaultPOPath(*TargetSettings, Culture->GetName()));
+		const FString DefaultPath = FPaths::GetPath(LocalizationConfigurationScript::GetDefaultPOPath(LocalizationTarget, Culture->GetName()));
 
 		FText DialogTitle;
 		{
 			FFormatNamedArguments FormatArguments;
-			FormatArguments.Add(TEXT("TargetName"), FText::FromString(TargetSettings->Name));
+			FormatArguments.Add(TEXT("TargetName"), FText::FromString(LocalizationTarget->Settings.Name));
 			FormatArguments.Add(TEXT("CultureName"), FText::FromString(Culture->GetDisplayName()));
 			DialogTitle = FText::Format(LOCTEXT("ExportSpecificTranslationsForTargetDialogTitleFormat", "Export {CultureName} Translations for {TargetName} to Directory"), FormatArguments);
 		}
@@ -391,7 +390,7 @@ FReply SLocalizationTargetEditorCultureRow::Export()
 		TArray<FString> SaveFilenames;
 		if (DesktopPlatform->SaveFileDialog(ParentWindowWindowHandle, DialogTitle.ToString(), DefaultPath, DefaultFilename, FileTypes, 0, SaveFilenames))
 		{
-			LocalizationCommandletTasks::ExportCulture(ParentWindow.ToSharedRef(), *TargetSettings, Culture->GetName(), TOptional<FString>(SaveFilenames.Top()));
+			LocalizationCommandletTasks::ExportCulture(ParentWindow.ToSharedRef(), LocalizationTarget, Culture->GetName(), TOptional<FString>(SaveFilenames.Top()));
 		}
 	}
 
diff --git a/Engine/Source/Editor/LocalizationDashboard/Private/SLocalizationTargetStatusButton.cpp b/Engine/Source/Editor/LocalizationDashboard/Private/SLocalizationTargetStatusButton.cpp
--- a/Engine/Source/Editor/LocalizationDashboard/Private/SLocalizationTargetStatusButton.cpp
+++ b/Engine/Source/Editor/LocalizationDashboard/Private/SLocalizationTargetStatusButton.cpp
@@ -87,7 +87,7 @@ FReply SLocalizationTargetStatusButton::OnClicked()
 		{
 			// Generate conflict report.
 			const TSharedPtr<SWindow> ParentWindow = FSlateApplication::Get().FindWidgetWindow(AsShared());
-			if (LocalizationCommandletTasks::GenerateReportsForTarget(ParentWindow.ToSharedRef(), Target->Settings))
+			if (LocalizationCommandletTasks::GenerateReportsForTarget(ParentWindow.ToSharedRef(), Target))
 			{
 				Target->UpdateStatusFromConflictReport();
 			}
@@ -98,7 +98,7 @@ FReply SLocalizationTargetStatusButton::OnClicked()
 		break;
 	case ELocalizationTargetStatus::ConflictsPresent:
 		FString ReportString;
-		if (FFileHelper::LoadFileToString(ReportString, *LocalizationConfigurationScript::GetConflictReportPath(Target->Settings)))
+		if (FFileHelper::LoadFileToString(ReportString, *LocalizationConfigurationScript::GetConflictReportPath(Target)))
 		{
 			TSharedPtr<SWindow> ReportWindow =
 				SNew(SWindow)
