Stubbed logic for localization configuration script generation to use engine or project config/content directories based on whether a target is a member of the unreal engine target set or the project-specific target set.

diff --git a/Engine/Source/Editor/LocalizationDashboard/Private/LocalizationConfigurationScript.cpp b/Engine/Source/Editor/LocalizationDashboard/Private/LocalizationConfigurationScript.cpp
--- a/Engine/Source/Editor/LocalizationDashboard/Private/LocalizationConfigurationScript.cpp
+++ b/Engine/Source/Editor/LocalizationDashboard/Private/LocalizationConfigurationScript.cpp
@@ -4,6 +4,24 @@
 #include "LocalizationConfigurationScript.h"
 #include "Classes/LocalizationTarget.h"
 
+namespace
+{
+	bool IsMemberOfUE4TargetSet(const ULocalizationTarget* const Target)
+	{
+		return false;
+	}
+
+	FString GetConfigDir(const ULocalizationTarget* const Target)
+	{
+		return IsMemberOfUE4TargetSet(Target) ? FPaths::EngineConfigDir() : FPaths::GameConfigDir();
+	}
+
+	FString GetContentDir(const ULocalizationTarget* const Target)
+	{
+		return IsMemberOfUE4TargetSet(Target) ? FPaths::EngineContentDir() : FPaths::GameContentDir();
+	}
+}
+
 namespace LocalizationConfigurationScript
 {
 	FString MakePathRelativeToProjectDirectory(const FString& Path)
@@ -13,14 +31,14 @@ namespace LocalizationConfigurationScript
 		return Result;
 	}
 
-	FString GetScriptDirectory()
+	FString GetScriptDirectory(const ULocalizationTarget* const Target)
 	{
-		return FPaths::GameConfigDir() / TEXT("Localization");
+		return GetConfigDir(Target) / TEXT("Localization");
 	}
 
 	FString GetDataDirectory(const ULocalizationTarget* const Target)
 	{
-		return FPaths::GameContentDir() / TEXT("Localization") / Target->Settings.Name;
+		return GetContentDir(Target) / TEXT("Localization") / Target->Settings.Name;
 	}
 
 	TArray<FString> GetScriptPaths(const ULocalizationTarget* const Target)
@@ -72,11 +90,8 @@ namespace LocalizationConfigurationScript
 	{
 		FLocalizationConfigurationScript Script;
 
-		const bool HasSourceCode = IFileManager::Get().DirectoryExists( *FPaths::GameSourceDir() );
-
-		const FString ConfigDirRelativeToGameDir = MakePathRelativeToProjectDirectory(FPaths::GameConfigDir());
-		const FString SourceDirRelativeToGameDir = MakePathRelativeToProjectDirectory(FPaths::GameSourceDir());
-		const FString ContentDirRelativeToGameDir = MakePathRelativeToProjectDirectory(FPaths::GameContentDir());
+		const FString ConfigDirRelativeToGameDir = MakePathRelativeToProjectDirectory(GetConfigDir(Target));
+		const FString ContentDirRelativeToGameDir = MakePathRelativeToProjectDirectory(GetContentDir(Target));
 
 		// CommonSettings
 		{
@@ -222,14 +237,14 @@ namespace LocalizationConfigurationScript
 
 	FString GetGatherScriptPath(const ULocalizationTarget* const Target)
 	{
-		return GetScriptDirectory() / FString::Printf( TEXT("%s_Gather.%s"), *(Target->Settings.Name), TEXT("ini") );
+		return GetScriptDirectory(Target) / FString::Printf( TEXT("%s_Gather.%s"), *(Target->Settings.Name), TEXT("ini") );
 	}
 
 	FLocalizationConfigurationScript GenerateImportScript(const ULocalizationTarget* const Target, const TOptional<FString> CultureName, const TOptional<FString> OutputPathOverride)
 	{
 		FLocalizationConfigurationScript Script;
 
-		const FString ContentDirRelativeToGameDir = MakePathRelativeToProjectDirectory(FPaths::GameContentDir());
+		const FString ContentDirRelativeToGameDir = MakePathRelativeToProjectDirectory(GetContentDir(Target));
 
 		// GatherTextStep0 - InternationalizationExport
 		{
@@ -314,7 +329,7 @@ namespace LocalizationConfigurationScript
 
 	FString GetImportScriptPath(const ULocalizationTarget* const Target, const TOptional<FString> CultureName)
 	{
-		const FString ConfigFileDirectory = GetScriptDirectory();
+		const FString ConfigFileDirectory = GetScriptDirectory(Target);
 		FString ConfigFilePath;
 		if (CultureName.IsSet())
 		{
@@ -331,7 +346,7 @@ namespace LocalizationConfigurationScript
 	{
 		FLocalizationConfigurationScript Script;
 
-		const FString ContentDirRelativeToGameDir = MakePathRelativeToProjectDirectory(FPaths::GameContentDir());
+		const FString ContentDirRelativeToGameDir = MakePathRelativeToProjectDirectory(GetContentDir(Target));
 
 		// GatherTextStep0 - InternationalizationExport
 		{
@@ -424,7 +439,7 @@ namespace LocalizationConfigurationScript
 
 	FString GetExportScriptPath(const ULocalizationTarget* const Target, const TOptional<FString> CultureName)
 	{
-		const FString ConfigFileDirectory = GetScriptDirectory();
+		const FString ConfigFileDirectory = GetScriptDirectory(Target);
 		FString ConfigFilePath;
 		if (CultureName.IsSet())
 		{
@@ -441,7 +456,7 @@ namespace LocalizationConfigurationScript
 	{
 		FLocalizationConfigurationScript Script;
 
-		const FString ContentDirRelativeToGameDir = MakePathRelativeToProjectDirectory(FPaths::GameContentDir());
+		const FString ContentDirRelativeToGameDir = MakePathRelativeToProjectDirectory(GetContentDir(Target));
 
 		// GatherTextStep0 - GenerateTextLocalizationReport
 		{
@@ -478,14 +493,14 @@ namespace LocalizationConfigurationScript
 
 	FString GetReportScriptPath(const ULocalizationTarget* const Target)
 	{
-		return GetScriptDirectory() / FString::Printf( TEXT("%s_GenerateReports.%s"), *(Target->Settings.Name), TEXT("ini") );
+		return GetScriptDirectory(Target) / FString::Printf( TEXT("%s_GenerateReports.%s"), *(Target->Settings.Name), TEXT("ini") );
 	}
 
 	FLocalizationConfigurationScript GenerateCompileScript(const ULocalizationTarget* const Target)
 	{
 		FLocalizationConfigurationScript Script;
 
-		const FString ContentDirRelativeToGameDir = MakePathRelativeToProjectDirectory(FPaths::GameContentDir());
+		const FString ContentDirRelativeToGameDir = MakePathRelativeToProjectDirectory(GetContentDir(Target));
 
 		// GatherTextStep0 - GenerateTextLocalizationResource
 		{
@@ -522,6 +537,6 @@ namespace LocalizationConfigurationScript
 
 	FString GetCompileScriptPath(const ULocalizationTarget* const Target)
 	{
-		return GetScriptDirectory() / FString::Printf( TEXT("%s_Compile.%s"), *(Target->Settings.Name), TEXT("ini") );
+		return GetScriptDirectory(Target) / FString::Printf( TEXT("%s_Compile.%s"), *(Target->Settings.Name), TEXT("ini") );
 	}
 }
\ No newline at end of file
