Refactored commandlet helper method to provide the necessary arguments for process creation, rather than actually creating the process. The localization commandlet executor UI now displays the commandlet process arguments as a tooltip for the task's name in the task list.

diff --git a/Engine/Source/Editor/LocalizationDashboard/Private/LocalizationCommandletExecution.cpp b/Engine/Source/Editor/LocalizationDashboard/Private/LocalizationCommandletExecution.cpp
--- a/Engine/Source/Editor/LocalizationDashboard/Private/LocalizationCommandletExecution.cpp
+++ b/Engine/Source/Editor/LocalizationDashboard/Private/LocalizationCommandletExecution.cpp
@@ -38,6 +38,7 @@ namespace
 			LocalizationCommandletExecution::FTask Task;
 			EState State;
 			FString LogOutput;
+			FString ProcessArguments;
 		};
 
 		friend class STaskRow;
@@ -65,6 +66,7 @@ namespace
 		TSharedRef<ITableRow> OnGenerateTaskListRow(TSharedPtr<FTaskListModel> TaskListModel, const TSharedRef<STableViewBase>& Table);
 		TSharedPtr<FTaskListModel> GetCurrentTaskToView() const;
 
+		FText GetCurrentTaskProcessArguments() const;
 		FText GetLogString() const;
 
 		FReply OnCopyLogClicked();
@@ -192,18 +194,18 @@ namespace
 									.VScrollBar(VerticalScrollBar)
 								]
 								+SVerticalBox::Slot()
-									.AutoHeight()
-									[
-										SAssignNew(HorizontalScrollBar, SScrollBar)
-										.Orientation(EOrientation::Orient_Horizontal)
-									]
-							]
-							+SHorizontalBox::Slot()
-								.AutoWidth()
+								.AutoHeight()
 								[
-									SAssignNew(VerticalScrollBar, SScrollBar)
-									.Orientation(EOrientation::Orient_Vertical)
+									SAssignNew(HorizontalScrollBar, SScrollBar)
+									.Orientation(EOrientation::Orient_Horizontal)
 								]
+							]
+							+SHorizontalBox::Slot()
+							.AutoWidth()
+							[
+								SAssignNew(VerticalScrollBar, SScrollBar)
+								.Orientation(EOrientation::Orient_Vertical)
+							]
 						]
 					]
 				+ SVerticalBox::Slot()
@@ -328,6 +330,7 @@ namespace
 		if (CommandletProcess.IsValid())
 		{
 			TaskListModel->State = FTaskListModel::EState::InProgress;
+			TaskListModel->ProcessArguments = CommandletProcess->GetProcessArguments();
 		}
 		else
 		{
@@ -483,7 +486,8 @@ namespace
 		else if (ColumnName == "TaskName")
 		{
 			return SNew(STextBlock)
-				.Text(TaskListModel->Task.Name);
+				.Text(TaskListModel->Task.Name)
+				.ToolTipText_Lambda( [this]{ return FText::FromString(TaskListModel->ProcessArguments); } );
 		}
 		else
 		{
@@ -560,6 +564,12 @@ namespace
 		return nullptr;
 	}
 
+	FText SLocalizationCommandletExecutor::GetCurrentTaskProcessArguments() const
+	{
+		const TSharedPtr<SLocalizationCommandletExecutor::FTaskListModel> TaskToView = GetCurrentTaskToView();
+		return TaskToView.IsValid() ? FText::FromString(TaskToView->ProcessArguments) : FText::GetEmpty();
+	}
+
 	FText SLocalizationCommandletExecutor::GetLogString() const
 	{
 		const TSharedPtr<SLocalizationCommandletExecutor::FTaskListModel> TaskToView = GetCurrentTaskToView();
@@ -676,7 +686,8 @@ TSharedPtr<FLocalizationCommandletProcess> FLocalizationCommandletProcess::Execu
 		}
 	}
 
-	FProcHandle CommandletProcessHandle = CommandletHelpers::CreateCommandletProcess(TEXT("GatherText"), *FPaths::GetProjectFilePath(), *CommandletArguments, true, true, true, nullptr, 0, nullptr, WritePipe);
+	const FString ProcessArguments = CommandletHelpers::BuildCommandletProcessArguments(TEXT("GatherText"), *FPaths::GetProjectFilePath(), *CommandletArguments);
+	FProcHandle CommandletProcessHandle = FPlatformProcess::CreateProc(*FUnrealEdMisc::Get().GetExecutableForCommandlets(), *ProcessArguments, true, true, true, nullptr, 0, nullptr, WritePipe);
 
 	// Close pipes if process failed.
 	if (!CommandletProcessHandle.IsValid())
@@ -685,7 +696,7 @@ TSharedPtr<FLocalizationCommandletProcess> FLocalizationCommandletProcess::Execu
 		return nullptr;
 	}
 
-	return MakeShareable(new FLocalizationCommandletProcess(ReadPipe, WritePipe, CommandletProcessHandle));
+	return MakeShareable(new FLocalizationCommandletProcess(ReadPipe, WritePipe, CommandletProcessHandle, ProcessArguments));
 }
 
 FLocalizationCommandletProcess::~FLocalizationCommandletProcess()
diff --git a/Engine/Source/Editor/LocalizationDashboard/Private/LocalizationCommandletExecution.h b/Engine/Source/Editor/LocalizationDashboard/Private/LocalizationCommandletExecution.h
--- a/Engine/Source/Editor/LocalizationDashboard/Private/LocalizationCommandletExecution.h
+++ b/Engine/Source/Editor/LocalizationDashboard/Private/LocalizationCommandletExecution.h
@@ -27,10 +27,11 @@ public:
 	static TSharedPtr<FLocalizationCommandletProcess> Execute(const FString& ConfigFilePath);
 
 private:
-	FLocalizationCommandletProcess(void* const InReadPipe, void* const InWritePipe, const FProcHandle InProcessHandle)
+	FLocalizationCommandletProcess(void* const InReadPipe, void* const InWritePipe, const FProcHandle InProcessHandle, const FString& InProcessArguments)
 		: ReadPipe(InReadPipe)
 		, WritePipe(InWritePipe)
 		, ProcessHandle(InProcessHandle)
+		, ProcessArguments(InProcessArguments)
 	{
 	}
 
@@ -47,8 +48,14 @@ public:
 		return ProcessHandle;
 	}
 
+	const FString& GetProcessArguments() const
+	{
+		return ProcessArguments;
+	}
+
 private:
 	void* const ReadPipe;
 	void* const WritePipe;
 	FProcHandle ProcessHandle;
+	FString ProcessArguments;
 };
\ No newline at end of file
diff --git a/Engine/Source/Editor/UnrealEd/Private/Commandlets/CommandletHelpers.cpp b/Engine/Source/Editor/UnrealEd/Private/Commandlets/CommandletHelpers.cpp
--- a/Engine/Source/Editor/UnrealEd/Private/Commandlets/CommandletHelpers.cpp
+++ b/Engine/Source/Editor/UnrealEd/Private/Commandlets/CommandletHelpers.cpp
@@ -5,29 +5,21 @@
 
 namespace CommandletHelpers
 {
-	FProcHandle CreateCommandletProcess(const TCHAR* CommandletName, const TCHAR* ProjectPath, const TCHAR* Parms, bool bLaunchDetached, bool bLaunchHidden, bool bLaunchReallyHidden, uint32* OutProcessID, int32 PriorityModifier, const TCHAR* OptionalWorkingDirectory, void* PipeWrite)
+	FString BuildCommandletProcessArguments(const TCHAR* const CommandletName, const TCHAR* const ProjectPath, const TCHAR* const AdditionalArguments)
 	{
-		const FString ExecutableURL = FUnrealEdMisc::Get().GetExecutableForCommandlets();
+		check(CommandletName && FCString::Strlen(CommandletName) != 0);
 
-		FProcHandle Result;
-
-		// Construct arguments string.
-		if (CommandletName && FCString::Strlen(CommandletName) != 0)
+		FString Arguments;
+		const FString RunCommand = FString::Printf( TEXT("-run=%s"), CommandletName );
+		if (ProjectPath && FCString::Strlen(ProjectPath) != 0)
 		{
-			FString Arguments;
-			const FString RunCommand = FString::Printf( TEXT("-run=%s"), CommandletName );
-			if (ProjectPath && FCString::Strlen(ProjectPath) != 0)
-			{
-				Arguments = FString::Printf( TEXT("%s %s"), ProjectPath, *RunCommand );
-			}
-			if (Parms && FCString::Strlen(Parms) != 0)
-			{
-				Arguments = FString::Printf( TEXT("%s %s"), *Arguments, Parms );
-			}
-
-			Result = FPlatformProcess::CreateProc(*ExecutableURL, *Arguments, bLaunchDetached, bLaunchHidden, bLaunchReallyHidden, OutProcessID, PriorityModifier, OptionalWorkingDirectory, PipeWrite);
+			Arguments = FString::Printf( TEXT("%s %s"), ProjectPath, *RunCommand );
+		}
+		if (AdditionalArguments && FCString::Strlen(AdditionalArguments) != 0)
+		{
+			Arguments = FString::Printf( TEXT("%s %s"), *Arguments, AdditionalArguments );
 		}
 
-		return Result;
+		return Arguments;
 	}
 }
\ No newline at end of file
diff --git a/Engine/Source/Editor/UnrealEd/Public/Commandlets/CommandletHelpers.h b/Engine/Source/Editor/UnrealEd/Public/Commandlets/CommandletHelpers.h
--- a/Engine/Source/Editor/UnrealEd/Public/Commandlets/CommandletHelpers.h
+++ b/Engine/Source/Editor/UnrealEd/Public/Commandlets/CommandletHelpers.h
@@ -4,5 +4,5 @@
 
 namespace CommandletHelpers
 {
-	UNREALED_API FProcHandle CreateCommandletProcess(const TCHAR* CommandletName, const TCHAR* ProjectPath, const TCHAR* Parms, bool bLaunchDetached, bool bLaunchHidden, bool bLaunchReallyHidden, uint32* OutProcessID, int32 PriorityModifier, const TCHAR* OptionalWorkingDirectory, void* PipeWrite);
+	UNREALED_API FString BuildCommandletProcessArguments(const TCHAR* const CommandletName, const TCHAR* const ProjectPath, const TCHAR* const AdditionalArguments);
 }
\ No newline at end of file
