Fixed missing header inclusion for localization commandlet execution using FUnrealEdMisc.

diff --git a/Engine/Source/Editor/LocalizationDashboard/Private/LocalizationCommandletExecution.cpp b/Engine/Source/Editor/LocalizationDashboard/Private/LocalizationCommandletExecution.cpp
--- a/Engine/Source/Editor/LocalizationDashboard/Private/LocalizationCommandletExecution.cpp
+++ b/Engine/Source/Editor/LocalizationDashboard/Private/LocalizationCommandletExecution.cpp
@@ -9,6 +9,7 @@
 #include "SThrobber.h"
 #include "Commandlets/CommandletHelpers.h"
 #include "EditorStyleSet.h"
+#include "UnrealEdMisc.h"
 
 #define LOCTEXT_NAMESPACE "LocalizationCommandletExecutor"
 
