Potential solution for an odd inclusion error regarding UnrealEdMisc.h being included from localization dashboard source files.

diff --git a/Engine/Source/Editor/LocalizationDashboard/Private/LocalizationDashboardPrivatePCH.h b/Engine/Source/Editor/LocalizationDashboard/Private/LocalizationDashboardPrivatePCH.h
--- a/Engine/Source/Editor/LocalizationDashboard/Private/LocalizationDashboardPrivatePCH.h
+++ b/Engine/Source/Editor/LocalizationDashboard/Private/LocalizationDashboardPrivatePCH.h
@@ -12,3 +12,4 @@
 #include "EditorStyle.h"
 #include "FileHelpers.h"
 #include "SlateBasics.h"
+#include "UnrealEd.h"
\ No newline at end of file
