Added engine localization  target set and project localization target set separation in localization dashboarded. Added tab for engine localization target set. Made the localization commandlet process use the project file as necessary.

diff --git a/Engine/Source/Editor/LocalizationDashboard/Classes/LocalizationTarget.h b/Engine/Source/Editor/LocalizationDashboard/Classes/LocalizationTarget.h
--- a/Engine/Source/Editor/LocalizationDashboard/Classes/LocalizationTarget.h
+++ b/Engine/Source/Editor/LocalizationDashboard/Classes/LocalizationTarget.h
@@ -140,13 +140,14 @@ public:
 	virtual void PostEditChangeProperty(FPropertyChangedEvent& PropertyChangedEvent) override;
 #endif
 
+	bool IsMemberOfEngineTargetSet() const;
 	bool UpdateWordCountsFromCSV();
 	void UpdateStatusFromConflictReport();
 	bool RenameTargetAndFiles(const FString& NewName);
 	bool DeleteFiles(const FString* const Culture = nullptr) const;
 };
 
-UCLASS(Config=Game, defaultconfig)
+UCLASS(Config=Localization, perObjectConfig, defaultconfig)
 class ULocalizationTargetSet : public UObject
 {
 	GENERATED_BODY()
@@ -155,6 +156,9 @@ public:
 	UPROPERTY(EditAnywhere, Category = "Targets")
 	TArray<ULocalizationTarget*> TargetObjects;
 
+	static FName EngineTargetSetName;
+	static FName ProjectTargetSetName;
+
 private:
 	UPROPERTY(config)
 	TArray<FLocalizationTargetSettings> Targets;
diff --git a/Engine/Source/Editor/LocalizationDashboard/LocalizationDashboard.Build.cs b/Engine/Source/Editor/LocalizationDashboard/LocalizationDashboard.Build.cs
--- a/Engine/Source/Editor/LocalizationDashboard/LocalizationDashboard.Build.cs
+++ b/Engine/Source/Editor/LocalizationDashboard/LocalizationDashboard.Build.cs
@@ -26,7 +26,8 @@ public class LocalizationDashboard : ModuleRules
                 "EditorStyle",
 				"DesktopPlatform",
                 "TranslationEditor",
-                "MainFrame"
+                "MainFrame",
+                "SourceControl"
 			}
 		);
 
diff --git a/Engine/Source/Editor/LocalizationDashboard/Private/LocalizationCommandletExecution.cpp b/Engine/Source/Editor/LocalizationDashboard/Private/LocalizationCommandletExecution.cpp
--- a/Engine/Source/Editor/LocalizationDashboard/Private/LocalizationCommandletExecution.cpp
+++ b/Engine/Source/Editor/LocalizationDashboard/Private/LocalizationCommandletExecution.cpp
@@ -10,6 +10,7 @@
 #include "Commandlets/CommandletHelpers.h"
 #include "EditorStyleSet.h"
 #include "UnrealEdMisc.h"
+#include "SourceControlHelpers.h"
 
 #define LOCTEXT_NAMESPACE "LocalizationCommandletExecutor"
 
@@ -195,18 +196,18 @@ namespace
 									.VScrollBar(VerticalScrollBar)
 								]
 								+SVerticalBox::Slot()
-								.AutoHeight()
-								[
-									SAssignNew(HorizontalScrollBar, SScrollBar)
-									.Orientation(EOrientation::Orient_Horizontal)
-								]
+									.AutoHeight()
+									[
+										SAssignNew(HorizontalScrollBar, SScrollBar)
+										.Orientation(EOrientation::Orient_Horizontal)
+									]
 							]
 							+SHorizontalBox::Slot()
-							.AutoWidth()
-							[
-								SAssignNew(VerticalScrollBar, SScrollBar)
-								.Orientation(EOrientation::Orient_Vertical)
-							]
+								.AutoWidth()
+								[
+									SAssignNew(VerticalScrollBar, SScrollBar)
+									.Orientation(EOrientation::Orient_Vertical)
+								]
 						]
 					]
 				+ SVerticalBox::Slot()
@@ -304,6 +305,18 @@ namespace
 
 		// Handle return code.
 		TSharedPtr<FTaskListModel> CurrentTaskModel = TaskListModels[CurrentTaskIndex];
+
+		// Restore engine's source control settings if necessary.
+		if (!CurrentTaskModel->Task.ShouldUseProjectFile)
+		{
+			const FString& EngineIniFile = SourceControlHelpers::GetGlobalSettingsIni();
+			const FString BackupEngineIniFile = FPaths::EngineSavedDir() / FPaths::GetCleanFilename(EngineIniFile) + TEXT(".bak");
+			if(!IFileManager::Get().Move(*EngineIniFile, *BackupEngineIniFile))
+			{
+				// TODO: Error failed to restore engine source control settings INI.
+			}
+		}
+
 		// Zero code is successful.
 		if (ReturnCode == 0)
 		{
@@ -327,7 +340,29 @@ namespace
 
 	void SLocalizationCommandletExecutor::ExecuteCommandlet(const TSharedRef<FTaskListModel>& TaskListModel)
 	{
-		CommandletProcess = FLocalizationCommandletProcess::Execute(TaskListModel->Task.ScriptPath);
+		// Handle source control settings if not using project file for commandlet executable process.
+		if (!TaskListModel->Task.ShouldUseProjectFile)
+		{
+			const FString& EngineIniFile = SourceControlHelpers::GetGlobalSettingsIni();
+
+			// Backup engine's source control settings.
+			const FString BackupEngineIniFile = FPaths::EngineSavedDir() / FPaths::GetCleanFilename(EngineIniFile) + TEXT(".bak");
+			if (COPY_OK == IFileManager::Get().Copy(*BackupEngineIniFile, *EngineIniFile))
+			{
+				// Replace engine's source control settings with project's.
+				const FString& ProjectIniFile = SourceControlHelpers::GetSettingsIni();
+				if (COPY_OK == IFileManager::Get().Copy(*EngineIniFile, *ProjectIniFile))
+				{
+					// TODO: Error failed to overwrite engine source control settings INI.
+				}
+			}
+			else
+			{
+				// TODO: Error failed to backup engine source control settings INI.
+			}
+		}
+
+		CommandletProcess = FLocalizationCommandletProcess::Execute(TaskListModel->Task.ScriptPath, TaskListModel->Task.ShouldUseProjectFile);
 		if (CommandletProcess.IsValid())
 		{
 			TaskListModel->State = FTaskListModel::EState::InProgress;
@@ -662,7 +697,7 @@ bool LocalizationCommandletExecution::Execute(const TSharedRef<SWindow>& ParentW
 	return CommandletExecutor->WasSuccessful();
 }
 
-TSharedPtr<FLocalizationCommandletProcess> FLocalizationCommandletProcess::Execute(const FString& ConfigFilePath)
+TSharedPtr<FLocalizationCommandletProcess> FLocalizationCommandletProcess::Execute(const FString& ConfigFilePath, const bool UseProjectFile)
 {
 	// Create pipes.
 	void* ReadPipe;
@@ -675,7 +710,7 @@ TSharedPtr<FLocalizationCommandletProcess> FLocalizationCommandletProcess::Execu
 	// Create process.
 	FString CommandletArguments;
 
-	const FString ConfigFileRelativeToGameDir = LocalizationConfigurationScript::MakePathRelativeToProjectDirectory(ConfigFilePath);
+	const FString ConfigFileRelativeToGameDir = LocalizationConfigurationScript::MakePathRelativeForCommandletProcess(ConfigFilePath, UseProjectFile);
 	CommandletArguments = FString::Printf( TEXT("-config=\"%s\""), *ConfigFileRelativeToGameDir );
 
 	ISourceControlModule& SourceControlModule = ISourceControlModule::Get();
@@ -687,7 +722,7 @@ TSharedPtr<FLocalizationCommandletProcess> FLocalizationCommandletProcess::Execu
 		}
 	}
 
-	const FString ProcessArguments = CommandletHelpers::BuildCommandletProcessArguments(TEXT("GatherText"), *FPaths::GetProjectFilePath(), *CommandletArguments);
+	const FString ProcessArguments = CommandletHelpers::BuildCommandletProcessArguments(TEXT("GatherText"), UseProjectFile ? *FPaths::GetProjectFilePath() : nullptr, *CommandletArguments);
 	FProcHandle CommandletProcessHandle = FPlatformProcess::CreateProc(*FUnrealEdMisc::Get().GetExecutableForCommandlets(), *ProcessArguments, true, true, true, nullptr, 0, nullptr, WritePipe);
 
 	// Close pipes if process failed.
diff --git a/Engine/Source/Editor/LocalizationDashboard/Private/LocalizationCommandletExecution.h b/Engine/Source/Editor/LocalizationDashboard/Private/LocalizationCommandletExecution.h
--- a/Engine/Source/Editor/LocalizationDashboard/Private/LocalizationCommandletExecution.h
+++ b/Engine/Source/Editor/LocalizationDashboard/Private/LocalizationCommandletExecution.h
@@ -9,13 +9,15 @@ namespace LocalizationCommandletExecution
 	{
 		FTask() {}
 
-		FTask(const FText& InName, const FString& InScriptPath)
+		FTask(const FText& InName, const FString& InScriptPath, const bool InShouldUseProjectFile)
 			: Name(InName)
 			, ScriptPath(InScriptPath)
+			, ShouldUseProjectFile(InShouldUseProjectFile)
 		{}
 
 		FText Name;
 		FString ScriptPath;
+		bool ShouldUseProjectFile;
 	};
 
 	bool Execute(const TSharedRef<SWindow>& ParentWindow, const FText& Title, const TArray<FTask>& Tasks);
@@ -24,7 +26,7 @@ namespace LocalizationCommandletExecution
 class FLocalizationCommandletProcess : public TSharedFromThis<FLocalizationCommandletProcess>
 {
 public:
-	static TSharedPtr<FLocalizationCommandletProcess> Execute(const FString& ConfigFilePath);
+	static TSharedPtr<FLocalizationCommandletProcess> Execute(const FString& ConfigFilePath, const bool UseProjectFile = true);
 
 private:
 	FLocalizationCommandletProcess(void* const InReadPipe, void* const InWritePipe, const FProcHandle InProcessHandle, const FString& InProcessArguments)
diff --git a/Engine/Source/Editor/LocalizationDashboard/Private/LocalizationCommandletTasks.cpp b/Engine/Source/Editor/LocalizationDashboard/Private/LocalizationCommandletTasks.cpp
--- a/Engine/Source/Editor/LocalizationDashboard/Private/LocalizationCommandletTasks.cpp
+++ b/Engine/Source/Editor/LocalizationDashboard/Private/LocalizationCommandletTasks.cpp
@@ -13,18 +13,20 @@ bool LocalizationCommandletTasks::GatherTargets(const TSharedRef<SWindow>& Paren
 
 	for (ULocalizationTarget* Target : Targets)
 	{
+		const bool ShouldUseProjectFile = !Target->IsMemberOfEngineTargetSet();
+
 		FFormatNamedArguments Arguments;
 		Arguments.Add(TEXT("TargetName"), FText::FromString(Target->Settings.Name));
 
 		const FText GatherTaskName = FText::Format(LOCTEXT("GatherTaskNameFormat", "Gather Text for {TargetName}"), Arguments);
 		const FString GatherScriptPath = LocalizationConfigurationScript::GetGatherScriptPath(Target);
 		LocalizationConfigurationScript::GenerateGatherScript(Target).Write(GatherScriptPath);
-		Tasks.Add(LocalizationCommandletExecution::FTask(GatherTaskName, GatherScriptPath));
+		Tasks.Add(LocalizationCommandletExecution::FTask(GatherTaskName, GatherScriptPath, ShouldUseProjectFile));
 
 		const FText ReportTaskName = FText::Format(LOCTEXT("ReportTaskNameFormat", "Generate Reports for {TargetName}"), Arguments);
 		const FString ReportScriptPath = LocalizationConfigurationScript::GetReportScriptPath(Target);
 		LocalizationConfigurationScript::GenerateReportScript(Target).Write(ReportScriptPath);
-		Tasks.Add(LocalizationCommandletExecution::FTask(ReportTaskName, ReportScriptPath));
+		Tasks.Add(LocalizationCommandletExecution::FTask(ReportTaskName, ReportScriptPath, ShouldUseProjectFile));
 	}
 
 	return LocalizationCommandletExecution::Execute(ParentWindow, LOCTEXT("GatherAllTargetsWindowTitle", "Gather Text for All Targets"), Tasks);
@@ -33,14 +35,15 @@ bool LocalizationCommandletTasks::GatherTargets(const TSharedRef<SWindow>& Paren
 bool LocalizationCommandletTasks::GatherTarget(const TSharedRef<SWindow>& ParentWindow, ULocalizationTarget* const Target)
 {
 	TArray<LocalizationCommandletExecution::FTask> Tasks;
+	const bool ShouldUseProjectFile = !Target->IsMemberOfEngineTargetSet();
 
 	const FString GatherScriptPath = LocalizationConfigurationScript::GetGatherScriptPath(Target);
 	LocalizationConfigurationScript::GenerateGatherScript(Target).Write(GatherScriptPath);
-	Tasks.Add(LocalizationCommandletExecution::FTask(LOCTEXT("GatherTaskName", "Gather Text"), GatherScriptPath));
+	Tasks.Add(LocalizationCommandletExecution::FTask(LOCTEXT("GatherTaskName", "Gather Text"), GatherScriptPath, ShouldUseProjectFile));
 
 	const FString ReportScriptPath = LocalizationConfigurationScript::GetReportScriptPath(Target);
 	LocalizationConfigurationScript::GenerateReportScript(Target).Write(ReportScriptPath);
-	Tasks.Add(LocalizationCommandletExecution::FTask(LOCTEXT("ReportTaskName", "Generate Reports"), ReportScriptPath));
+	Tasks.Add(LocalizationCommandletExecution::FTask(LOCTEXT("ReportTaskName", "Generate Reports"), ReportScriptPath, ShouldUseProjectFile));
 
 	FFormatNamedArguments Arguments;
 	Arguments.Add(TEXT("TargetName"), FText::FromString(Target->Settings.Name));
@@ -56,6 +59,8 @@ bool LocalizationCommandletTasks::ImportTargets(const TSharedRef<SWindow>& Paren
 
 	for (ULocalizationTarget* Target : Targets)
 	{
+		const bool ShouldUseProjectFile = !Target->IsMemberOfEngineTargetSet();
+
 		FFormatNamedArguments Arguments;
 		Arguments.Add(TEXT("TargetName"), FText::FromString(Target->Settings.Name));
 
@@ -63,12 +68,12 @@ bool LocalizationCommandletTasks::ImportTargets(const TSharedRef<SWindow>& Paren
 		const FString ImportScriptPath = LocalizationConfigurationScript::GetImportScriptPath(Target, TOptional<FString>());
 		const TOptional<FString> DirectoryPathForTarget = DirectoryPath.IsSet() ? DirectoryPath.GetValue() / Target->Settings.Name : TOptional<FString>();
 		LocalizationConfigurationScript::GenerateImportScript(Target, TOptional<FString>(), DirectoryPathForTarget).Write(ImportScriptPath);
-		Tasks.Add(LocalizationCommandletExecution::FTask(ImportTaskName, ImportScriptPath));
+		Tasks.Add(LocalizationCommandletExecution::FTask(ImportTaskName, ImportScriptPath, ShouldUseProjectFile));
 
 		const FText ReportTaskName = FText::Format(LOCTEXT("ReportTaskNameFormat", "Generate Reports for {TargetName}"), Arguments);
 		const FString ReportScriptPath = LocalizationConfigurationScript::GetReportScriptPath(Target);
 		LocalizationConfigurationScript::GenerateReportScript(Target).Write(ReportScriptPath);
-		Tasks.Add(LocalizationCommandletExecution::FTask(ReportTaskName, ReportScriptPath));
+		Tasks.Add(LocalizationCommandletExecution::FTask(ReportTaskName, ReportScriptPath, ShouldUseProjectFile));
 	}
 
 	return LocalizationCommandletExecution::Execute(ParentWindow, LOCTEXT("ImportForAllTargetsWindowTitle", "Import Translations for All Targets"), Tasks);
@@ -77,14 +82,15 @@ bool LocalizationCommandletTasks::ImportTargets(const TSharedRef<SWindow>& Paren
 bool LocalizationCommandletTasks::ImportTarget(const TSharedRef<SWindow>& ParentWindow, ULocalizationTarget* const Target, const TOptional<FString> DirectoryPath)
 {
 	TArray<LocalizationCommandletExecution::FTask> Tasks;
+	const bool ShouldUseProjectFile = !Target->IsMemberOfEngineTargetSet();
 
 	const FString ImportScriptPath = LocalizationConfigurationScript::GetImportScriptPath(Target, TOptional<FString>());
 	LocalizationConfigurationScript::GenerateImportScript(Target, TOptional<FString>(), DirectoryPath).Write(ImportScriptPath);
-	Tasks.Add(LocalizationCommandletExecution::FTask(LOCTEXT("ImportTaskName", "Import Translations"), ImportScriptPath));
+	Tasks.Add(LocalizationCommandletExecution::FTask(LOCTEXT("ImportTaskName", "Import Translations"), ImportScriptPath, ShouldUseProjectFile));
 
 	const FString ReportScriptPath = LocalizationConfigurationScript::GetReportScriptPath(Target);
 	LocalizationConfigurationScript::GenerateReportScript(Target).Write(ReportScriptPath);
-	Tasks.Add(LocalizationCommandletExecution::FTask(LOCTEXT("ReportTaskName", "Generate Reports"), ReportScriptPath));
+	Tasks.Add(LocalizationCommandletExecution::FTask(LOCTEXT("ReportTaskName", "Generate Reports"), ReportScriptPath, ShouldUseProjectFile));
 
 	FFormatNamedArguments Arguments;
 	Arguments.Add(TEXT("TargetName"), FText::FromString(Target->Settings.Name));
@@ -102,15 +108,16 @@ bool LocalizationCommandletTasks::ImportCulture(const TSharedRef<SWindow>& Paren
 	}
 
 	TArray<LocalizationCommandletExecution::FTask> Tasks;
+	const bool ShouldUseProjectFile = !Target->IsMemberOfEngineTargetSet();
 
 	const FString DefaultImportScriptPath = LocalizationConfigurationScript::GetImportScriptPath(Target, TOptional<FString>(CultureName));
 	const FString ImportScriptPath = FPaths::CreateTempFilename(*FPaths::GetPath(DefaultImportScriptPath), *FPaths::GetBaseFilename(DefaultImportScriptPath), *FPaths::GetExtension(DefaultImportScriptPath, true));
 	LocalizationConfigurationScript::GenerateImportScript(Target, TOptional<FString>(CultureName), FilePath).Write(ImportScriptPath);
-	Tasks.Add(LocalizationCommandletExecution::FTask(LOCTEXT("ImportTaskName", "Import Translations"), ImportScriptPath));
+	Tasks.Add(LocalizationCommandletExecution::FTask(LOCTEXT("ImportTaskName", "Import Translations"), ImportScriptPath, ShouldUseProjectFile));
 
 	const FString ReportScriptPath = LocalizationConfigurationScript::GetReportScriptPath(Target);
 	LocalizationConfigurationScript::GenerateReportScript(Target).Write(ReportScriptPath);
-	Tasks.Add(LocalizationCommandletExecution::FTask(LOCTEXT("ReportTaskName", "Generate Reports"), ReportScriptPath));
+	Tasks.Add(LocalizationCommandletExecution::FTask(LOCTEXT("ReportTaskName", "Generate Reports"), ReportScriptPath, ShouldUseProjectFile));
 
 	FFormatNamedArguments Arguments;
 	Arguments.Add(TEXT("CultureName"), FText::FromString(Culture->GetDisplayName()));
@@ -130,6 +137,8 @@ bool LocalizationCommandletTasks::ExportTargets(const TSharedRef<SWindow>& Paren
 
 	for (ULocalizationTarget* Target : Targets)
 	{
+		const bool ShouldUseProjectFile = !Target->IsMemberOfEngineTargetSet();
+
 		FFormatNamedArguments Arguments;
 		Arguments.Add(TEXT("TargetName"), FText::FromString(Target->Settings.Name));
 
@@ -137,7 +146,7 @@ bool LocalizationCommandletTasks::ExportTargets(const TSharedRef<SWindow>& Paren
 		const FString ExportScriptPath = LocalizationConfigurationScript::GetExportScriptPath(Target, TOptional<FString>());
 		const TOptional<FString> DirectoryPathForTarget = DirectoryPath.IsSet() ? DirectoryPath.GetValue() / Target->Settings.Name : TOptional<FString>();
 		LocalizationConfigurationScript::GenerateExportScript(Target, TOptional<FString>(), DirectoryPathForTarget).Write(ExportScriptPath);
-		Tasks.Add(LocalizationCommandletExecution::FTask(ExportTaskName, ExportScriptPath));
+		Tasks.Add(LocalizationCommandletExecution::FTask(ExportTaskName, ExportScriptPath, ShouldUseProjectFile));
 	}
 
 	return LocalizationCommandletExecution::Execute(ParentWindow, LOCTEXT("ExportForAllTargetsWindowTitle", "Export Translations for All Targets"), Tasks);
@@ -146,10 +155,11 @@ bool LocalizationCommandletTasks::ExportTargets(const TSharedRef<SWindow>& Paren
 bool LocalizationCommandletTasks::ExportTarget(const TSharedRef<SWindow>& ParentWindow, ULocalizationTarget* const Target, const TOptional<FString> DirectoryPath)
 {
 	TArray<LocalizationCommandletExecution::FTask> Tasks;
+	const bool ShouldUseProjectFile = !Target->IsMemberOfEngineTargetSet();
 
 	const FString ExportScriptPath = LocalizationConfigurationScript::GetExportScriptPath(Target, TOptional<FString>());
 	LocalizationConfigurationScript::GenerateExportScript(Target, TOptional<FString>(), DirectoryPath).Write(ExportScriptPath);
-	Tasks.Add(LocalizationCommandletExecution::FTask(LOCTEXT("ExportTaskName", "Export Translations"), ExportScriptPath));
+	Tasks.Add(LocalizationCommandletExecution::FTask(LOCTEXT("ExportTaskName", "Export Translations"), ExportScriptPath, ShouldUseProjectFile));
 
 	FFormatNamedArguments Arguments;
 	Arguments.Add(TEXT("TargetName"), FText::FromString(Target->Settings.Name));
@@ -167,11 +177,12 @@ bool LocalizationCommandletTasks::ExportCulture(const TSharedRef<SWindow>& Paren
 	}
 
 	TArray<LocalizationCommandletExecution::FTask> Tasks;
+	const bool ShouldUseProjectFile = !Target->IsMemberOfEngineTargetSet();
 
 	const FString DefaultExportScriptPath = LocalizationConfigurationScript::GetExportScriptPath(Target, TOptional<FString>(CultureName));
 	const FString ExportScriptPath = FPaths::CreateTempFilename(*FPaths::GetPath(DefaultExportScriptPath), *FPaths::GetBaseFilename(DefaultExportScriptPath), *FPaths::GetExtension(DefaultExportScriptPath, true));
 	LocalizationConfigurationScript::GenerateExportScript(Target, TOptional<FString>(CultureName), FilePath).Write(ExportScriptPath);
-	Tasks.Add(LocalizationCommandletExecution::FTask(LOCTEXT("ExportTaskName", "Export Translations"), ExportScriptPath));
+	Tasks.Add(LocalizationCommandletExecution::FTask(LOCTEXT("ExportTaskName", "Export Translations"), ExportScriptPath, ShouldUseProjectFile));
 
 	FFormatNamedArguments Arguments;
 	Arguments.Add(TEXT("CultureName"), FText::FromString(Culture->GetDisplayName()));
@@ -190,13 +201,15 @@ bool LocalizationCommandletTasks::GenerateReportsForTargets(const TSharedRef<SWi
 
 	for (ULocalizationTarget* Target : Targets)
 	{
+		const bool ShouldUseProjectFile = !Target->IsMemberOfEngineTargetSet();
+
 		FFormatNamedArguments Arguments;
 		Arguments.Add(TEXT("TargetName"), FText::FromString(Target->Settings.Name));
 
 		const FText ReportTaskName = FText::Format(LOCTEXT("ReportTaskNameFormat", "Generate Reports for {TargetName}"), Arguments);
 		const FString ReportScriptPath = LocalizationConfigurationScript::GetReportScriptPath(Target);
 		LocalizationConfigurationScript::GenerateReportScript(Target).Write(ReportScriptPath);
-		Tasks.Add(LocalizationCommandletExecution::FTask(ReportTaskName, ReportScriptPath));
+		Tasks.Add(LocalizationCommandletExecution::FTask(ReportTaskName, ReportScriptPath, ShouldUseProjectFile));
 	}
 
 	return LocalizationCommandletExecution::Execute(ParentWindow, LOCTEXT("GenerateReportsForAllTargetsWindowTitle", "Generate Reports for All Targets"), Tasks);
@@ -205,10 +218,11 @@ bool LocalizationCommandletTasks::GenerateReportsForTargets(const TSharedRef<SWi
 bool LocalizationCommandletTasks::GenerateReportsForTarget(const TSharedRef<SWindow>& ParentWindow, ULocalizationTarget* const Target)
 {
 	TArray<LocalizationCommandletExecution::FTask> Tasks;
+	const bool ShouldUseProjectFile = !Target->IsMemberOfEngineTargetSet();
 
 	const FString ReportScriptPath = LocalizationConfigurationScript::GetReportScriptPath(Target);
 	LocalizationConfigurationScript::GenerateReportScript(Target).Write(ReportScriptPath);
-	Tasks.Add(LocalizationCommandletExecution::FTask(LOCTEXT("ReportTaskName", "Generate Reports"), ReportScriptPath));
+	Tasks.Add(LocalizationCommandletExecution::FTask(LOCTEXT("ReportTaskName", "Generate Reports"), ReportScriptPath, ShouldUseProjectFile));
 
 	FFormatNamedArguments Arguments;
 	Arguments.Add(TEXT("TargetName"), FText::FromString(Target->Settings.Name));
@@ -223,13 +237,15 @@ bool LocalizationCommandletTasks::CompileTargets(const TSharedRef<SWindow>& Pare
 
 	for (ULocalizationTarget* Target : Targets)
 	{
+		const bool ShouldUseProjectFile = !Target->IsMemberOfEngineTargetSet();
+
 		FFormatNamedArguments Arguments;
 		Arguments.Add(TEXT("TargetName"), FText::FromString(Target->Settings.Name));
 
 		const FText CompileTaskName = FText::Format(LOCTEXT("CompileTaskNameFormat", "Compile Translations for {TargetName}"), Arguments);
 		const FString CompileScriptPath = LocalizationConfigurationScript::GetCompileScriptPath(Target);
 		LocalizationConfigurationScript::GenerateCompileScript(Target).Write(CompileScriptPath);
-		Tasks.Add(LocalizationCommandletExecution::FTask(CompileTaskName, CompileScriptPath));
+		Tasks.Add(LocalizationCommandletExecution::FTask(CompileTaskName, CompileScriptPath, ShouldUseProjectFile));
 	}
 
 	return LocalizationCommandletExecution::Execute(ParentWindow, LOCTEXT("GenerateLocResForAllTargetsWindowTitle", "Compile Translations for All Targets"), Tasks);
@@ -238,10 +254,11 @@ bool LocalizationCommandletTasks::CompileTargets(const TSharedRef<SWindow>& Pare
 bool LocalizationCommandletTasks::CompileTarget(const TSharedRef<SWindow>& ParentWindow, ULocalizationTarget* const Target)
 {
 	TArray<LocalizationCommandletExecution::FTask> Tasks;
+	const bool ShouldUseProjectFile = !Target->IsMemberOfEngineTargetSet();
 
 	const FString CompileScriptPath = LocalizationConfigurationScript::GetCompileScriptPath(Target);
 	LocalizationConfigurationScript::GenerateCompileScript(Target).Write(CompileScriptPath);
-	Tasks.Add(LocalizationCommandletExecution::FTask(LOCTEXT("CompileTaskName", "Compile Translations"), CompileScriptPath));
+	Tasks.Add(LocalizationCommandletExecution::FTask(LOCTEXT("CompileTaskName", "Compile Translations"), CompileScriptPath, ShouldUseProjectFile));
 
 	FFormatNamedArguments Arguments;
 	Arguments.Add(TEXT("TargetName"), FText::FromString(Target->Settings.Name));
diff --git a/Engine/Source/Editor/LocalizationDashboard/Private/LocalizationConfigurationScript.cpp b/Engine/Source/Editor/LocalizationDashboard/Private/LocalizationConfigurationScript.cpp
--- a/Engine/Source/Editor/LocalizationDashboard/Private/LocalizationConfigurationScript.cpp
+++ b/Engine/Source/Editor/LocalizationDashboard/Private/LocalizationConfigurationScript.cpp
@@ -6,28 +6,24 @@
 
 namespace
 {
-	bool IsMemberOfUE4TargetSet(const ULocalizationTarget* const Target)
-	{
-		return false;
-	}
-
 	FString GetConfigDir(const ULocalizationTarget* const Target)
 	{
-		return IsMemberOfUE4TargetSet(Target) ? FPaths::EngineConfigDir() : FPaths::GameConfigDir();
+		return Target->IsMemberOfEngineTargetSet() ? FPaths::EngineConfigDir() : FPaths::GameConfigDir();
 	}
 
 	FString GetContentDir(const ULocalizationTarget* const Target)
 	{
-		return IsMemberOfUE4TargetSet(Target) ? FPaths::EngineContentDir() : FPaths::GameContentDir();
+		return Target->IsMemberOfEngineTargetSet() ? FPaths::EngineContentDir() : FPaths::GameContentDir();
 	}
 }
 
 namespace LocalizationConfigurationScript
 {
-	FString MakePathRelativeToProjectDirectory(const FString& Path)
+	FString MakePathRelativeForCommandletProcess(const FString& Path, const bool IsUsingProjectFile)
 	{
 		FString Result = Path;
-		FPaths::MakePathRelativeTo(Result, *FPaths::GameDir());
+		const FString ProjectDir = !IsUsingProjectFile ? FPaths::EngineDir() : FPaths::GameDir();
+		FPaths::MakePathRelativeTo(Result, *ProjectDir);
 		return Result;
 	}
 
@@ -90,8 +86,8 @@ namespace LocalizationConfigurationScript
 	{
 		FLocalizationConfigurationScript Script;
 
-		const FString ConfigDirRelativeToGameDir = MakePathRelativeToProjectDirectory(GetConfigDir(Target));
-		const FString ContentDirRelativeToGameDir = MakePathRelativeToProjectDirectory(GetContentDir(Target));
+		const FString ConfigDirRelativeToGameDir = MakePathRelativeForCommandletProcess(GetConfigDir(Target), !Target->IsMemberOfEngineTargetSet());
+		const FString ContentDirRelativeToGameDir = MakePathRelativeForCommandletProcess(GetContentDir(Target), !Target->IsMemberOfEngineTargetSet());
 
 		// CommonSettings
 		{
@@ -103,13 +99,13 @@ namespace LocalizationConfigurationScript
 				ULocalizationTarget* const * OtherTarget = LocalizationTargetSet->TargetObjects.FindByPredicate([&TargetDependencyName](ULocalizationTarget* const OtherTarget)->bool{return OtherTarget->Settings.Name == TargetDependencyName;});
 				if (OtherTarget)
 				{
-					ConfigSection.Add( TEXT("ManifestDependencies"), MakePathRelativeToProjectDirectory(GetManifestPath(*OtherTarget)) );
+					ConfigSection.Add( TEXT("ManifestDependencies"), MakePathRelativeForCommandletProcess(GetManifestPath(*OtherTarget), !Target->IsMemberOfEngineTargetSet()) );
 				}
 			}
 
 			for (const FFilePath& Path : Target->Settings.AdditionalManifestDependencies)
 			{
-				ConfigSection.Add( TEXT("ManifestDependencies"), MakePathRelativeToProjectDirectory(Path.FilePath) );
+				ConfigSection.Add( TEXT("ManifestDependencies"), MakePathRelativeForCommandletProcess(Path.FilePath, !Target->IsMemberOfEngineTargetSet()) );
 			}
 
 			const FString SourcePath = ContentDirRelativeToGameDir / TEXT("Localization") / Target->Settings.Name;
@@ -144,7 +140,7 @@ namespace LocalizationConfigurationScript
 			}
 
 			// Exclude Paths
-			ConfigSection.Add( TEXT("ExcludePaths"), ConfigDirRelativeToGameDir / TEXT("Localization") );
+			ConfigSection.Add( TEXT("ExcludePaths"), TEXT("*/Config/Localization/*") );
 			for (const auto& ExcludePath : Target->Settings.GatherFromTextFiles.ExcludePathWildcards)
 			{
 				ConfigSection.Add( TEXT("ExcludePaths"), ExcludePath );
@@ -172,7 +168,7 @@ namespace LocalizationConfigurationScript
 			}
 
 			// Exclude Paths
-			ConfigSection.Add( TEXT("ExcludePaths"), ContentDirRelativeToGameDir / TEXT("Localization") );
+			ConfigSection.Add( TEXT("ExcludePaths"), TEXT("*/Content/Localization/*") );
 			for (const auto& ExcludePath : Target->Settings.GatherFromPackages.ExcludePathWildcards)
 			{
 				ConfigSection.Add( TEXT("ExcludePaths"), ExcludePath );
@@ -244,7 +240,7 @@ namespace LocalizationConfigurationScript
 	{
 		FLocalizationConfigurationScript Script;
 
-		const FString ContentDirRelativeToGameDir = MakePathRelativeToProjectDirectory(GetContentDir(Target));
+		const FString ContentDirRelativeToGameDir = MakePathRelativeForCommandletProcess(GetContentDir(Target), !Target->IsMemberOfEngineTargetSet());
 
 		// GatherTextStep0 - InternationalizationExport
 		{
@@ -262,12 +258,12 @@ namespace LocalizationConfigurationScript
 				// The output path for a specific culture is a file path.
 				if (CultureName.IsSet())
 				{
-					SourcePath = MakePathRelativeToProjectDirectory( FPaths::GetPath(OutputPathOverride.GetValue()) );
+					SourcePath = MakePathRelativeForCommandletProcess( FPaths::GetPath(OutputPathOverride.GetValue()), !Target->IsMemberOfEngineTargetSet() );
 				}
 				// Otherwise, it is a directory path.
 				else
 				{
-					SourcePath = MakePathRelativeToProjectDirectory( OutputPathOverride.GetValue() );
+					SourcePath = MakePathRelativeForCommandletProcess( OutputPathOverride.GetValue(), !Target->IsMemberOfEngineTargetSet() );
 				}
 			}
 			// Use the default PO file's directory path.
@@ -346,7 +342,7 @@ namespace LocalizationConfigurationScript
 	{
 		FLocalizationConfigurationScript Script;
 
-		const FString ContentDirRelativeToGameDir = MakePathRelativeToProjectDirectory(GetContentDir(Target));
+		const FString ContentDirRelativeToGameDir = MakePathRelativeForCommandletProcess(GetContentDir(Target), !Target->IsMemberOfEngineTargetSet());
 
 		// GatherTextStep0 - InternationalizationExport
 		{
@@ -367,12 +363,12 @@ namespace LocalizationConfigurationScript
 				// The output path for a specific culture is a file path.
 				if (CultureName.IsSet())
 				{
-					DestinationPath = MakePathRelativeToProjectDirectory( FPaths::GetPath(OutputPathOverride.GetValue()) );
+					DestinationPath = MakePathRelativeForCommandletProcess( FPaths::GetPath(OutputPathOverride.GetValue()), !Target->IsMemberOfEngineTargetSet() );
 				}
 				// Otherwise, it is a directory path.
 				else
 				{
-					DestinationPath = MakePathRelativeToProjectDirectory( OutputPathOverride.GetValue() );
+					DestinationPath = MakePathRelativeForCommandletProcess( OutputPathOverride.GetValue(), !Target->IsMemberOfEngineTargetSet() );
 				}
 			}
 			// Use the default PO file's directory path.
@@ -456,7 +452,7 @@ namespace LocalizationConfigurationScript
 	{
 		FLocalizationConfigurationScript Script;
 
-		const FString ContentDirRelativeToGameDir = MakePathRelativeToProjectDirectory(GetContentDir(Target));
+		const FString ContentDirRelativeToGameDir = MakePathRelativeForCommandletProcess(GetContentDir(Target), !Target->IsMemberOfEngineTargetSet());
 
 		// GatherTextStep0 - GenerateTextLocalizationReport
 		{
@@ -500,7 +496,7 @@ namespace LocalizationConfigurationScript
 	{
 		FLocalizationConfigurationScript Script;
 
-		const FString ContentDirRelativeToGameDir = MakePathRelativeToProjectDirectory(GetContentDir(Target));
+		const FString ContentDirRelativeToGameDir = MakePathRelativeForCommandletProcess(GetContentDir(Target), !Target->IsMemberOfEngineTargetSet());
 
 		// GatherTextStep0 - GenerateTextLocalizationResource
 		{
diff --git a/Engine/Source/Editor/LocalizationDashboard/Private/LocalizationConfigurationScript.h b/Engine/Source/Editor/LocalizationDashboard/Private/LocalizationConfigurationScript.h
--- a/Engine/Source/Editor/LocalizationDashboard/Private/LocalizationConfigurationScript.h
+++ b/Engine/Source/Editor/LocalizationDashboard/Private/LocalizationConfigurationScript.h
@@ -20,7 +20,7 @@ struct FLocalizationConfigurationScript : public FConfigFile
 
 namespace LocalizationConfigurationScript
 {
-	FString MakePathRelativeToProjectDirectory(const FString& Path);
+	FString MakePathRelativeForCommandletProcess(const FString& Path, const bool IsUsingProjectFile);
 
 	FString GetScriptDirectory(const ULocalizationTarget* const Target);
 	FString GetDataDirectory(const ULocalizationTarget* const Target);
diff --git a/Engine/Source/Editor/LocalizationDashboard/Private/LocalizationDashboard.cpp b/Engine/Source/Editor/LocalizationDashboard/Private/LocalizationDashboard.cpp
--- a/Engine/Source/Editor/LocalizationDashboard/Private/LocalizationDashboard.cpp
+++ b/Engine/Source/Editor/LocalizationDashboard/Private/LocalizationDashboard.cpp
@@ -20,11 +20,12 @@ public:
 	SLATE_BEGIN_ARGS(SLocalizationDashboard) {}
 	SLATE_END_ARGS()
 
-	void Construct(const FArguments& InArgs, const TSharedPtr<SWindow>& OwningWindow, const TSharedRef<SDockTab>& OwningTab, ULocalizationTargetSet* const TargetSet);
+	void Construct(const FArguments& InArgs, const TSharedPtr<SWindow>& OwningWindow, const TSharedRef<SDockTab>& OwningTab);
 
 	TWeakPtr<SDockTab> ShowTargetEditor(ULocalizationTarget* const LocalizationTarget);
 
 private:
+	static const FName EngineTargetSetDetailsTabName;
 	static const FName ProjectTargetSetDetailsTabName;
 	static const FName DocumentsTabName;
 
@@ -32,13 +33,12 @@ private:
 	TMap< TWeakObjectPtr<ULocalizationTarget>, TWeakPtr<SDockTab> > TargetToTabMap;
 };
 
+const FName SLocalizationDashboard::EngineTargetSetDetailsTabName("EngineTargets");
 const FName SLocalizationDashboard::ProjectTargetSetDetailsTabName("ProjectTargets");
 const FName SLocalizationDashboard::DocumentsTabName("Documents");
 
-void SLocalizationDashboard::Construct(const FArguments& InArgs, const TSharedPtr<SWindow>& OwningWindow, const TSharedRef<SDockTab>& OwningTab, ULocalizationTargetSet* const TargetSet)
+void SLocalizationDashboard::Construct(const FArguments& InArgs, const TSharedPtr<SWindow>& OwningWindow, const TSharedRef<SDockTab>& OwningTab)
 {
-	check(TargetSet);
-
 	TabManager = FGlobalTabmanager::Get()->NewTabManager(OwningTab);
 
 	const auto& PersistLayout = [](const TSharedRef<FTabManager::FLayout>& LayoutToSave)
@@ -47,10 +47,23 @@ void SLocalizationDashboard::Construct(const FArguments& InArgs, const TSharedPt
 	};
 	TabManager->SetOnPersistLayout(FTabManager::FOnPersistLayout::CreateLambda(PersistLayout));
 
-	const auto& CreateDetailsTab = [TargetSet](const FSpawnTabArgs& SpawnTabArgs) -> TSharedRef<SDockTab>
+	const auto& CreateDetailsTab = [](const FSpawnTabArgs& SpawnTabArgs) -> TSharedRef<SDockTab>
 	{
+		FText DockTabLabel;
+		ULocalizationTargetSet* TargetSet = nullptr;
+		if (SpawnTabArgs.GetTabId() == EngineTargetSetDetailsTabName)
+		{
+			DockTabLabel = LOCTEXT("EngineTargetsTabLabel", "Engine Targets");
+			TargetSet = FindObjectChecked<ULocalizationTargetSet>(GetTransientPackage(), TEXT("EngineLocalizationTargetSet"));
+		}
+		else if (SpawnTabArgs.GetTabId() == ProjectTargetSetDetailsTabName)
+		{
+			DockTabLabel = LOCTEXT("ProjectTargetsTabLabel", "Project Targets");
+			TargetSet = FindObjectChecked<ULocalizationTargetSet>(GetTransientPackage(), TEXT("ProjectLocalizationTargetSet"));
+		}
+
 		const TSharedRef<SDockTab> DockTab = SNew(SDockTab)
-			.Label(LOCTEXT("ProjectTargetsTabLabel", "Project Targets"));
+			.Label(DockTabLabel);
 
 		FPropertyEditorModule& PropertyModule = FModuleManager::LoadModuleChecked<FPropertyEditorModule>("PropertyEditor");
 		FDetailsViewArgs DetailsViewArgs(false, false, false, FDetailsViewArgs::ENameAreaSettings::HideNameArea, false, nullptr, false, NAME_None);
@@ -63,11 +76,14 @@ void SLocalizationDashboard::Construct(const FArguments& InArgs, const TSharedPt
 		return DockTab;
 	};
 	const TSharedRef<FWorkspaceItem> TargetSetsWorkspaceMenuCategory = TabManager->AddLocalWorkspaceMenuCategory(LOCTEXT("TargetSetsWorkspaceMenuCategory", "Target Sets"));
+	FTabSpawnerEntry& EngineTargetSetTabSpawnerEntry = TabManager->RegisterTabSpawner(EngineTargetSetDetailsTabName, FOnSpawnTab::CreateLambda(CreateDetailsTab))
+		.SetDisplayName(LOCTEXT("EngineTargetsDetailTabSpawner", "Engine Targets"));
+	TargetSetsWorkspaceMenuCategory->AddItem(EngineTargetSetTabSpawnerEntry.AsShared());
 	FTabSpawnerEntry& ProjectTargetSetTabSpawnerEntry = TabManager->RegisterTabSpawner(ProjectTargetSetDetailsTabName, FOnSpawnTab::CreateLambda(CreateDetailsTab))
 		.SetDisplayName(LOCTEXT("ProjectTargetsDetailTabSpawner", "Project Targets"));
 	TargetSetsWorkspaceMenuCategory->AddItem(ProjectTargetSetTabSpawnerEntry.AsShared());
 
-	TSharedRef<FTabManager::FLayout> Layout = FTabManager::NewLayout("LocalizationDashboard_Experimental_V3")
+	TSharedRef<FTabManager::FLayout> Layout = FTabManager::NewLayout("LocalizationDashboard_Experimental_V4")
 		->AddArea
 		(
 		FTabManager::NewPrimaryArea()
@@ -76,6 +92,7 @@ void SLocalizationDashboard::Construct(const FArguments& InArgs, const TSharedPt
 		(
 		FTabManager::NewStack()
 		->SetHideTabWell(true)
+		->AddTab(EngineTargetSetDetailsTabName, ETabState::ClosedTab)
 		->AddTab(ProjectTargetSetDetailsTabName, ETabState::OpenedTab)
 		)
 		->Split
@@ -195,7 +212,7 @@ void FLocalizationDashboard::RegisterTabSpawner()
 			.Label(LOCTEXT("MainTabTitle", "Localization Dashboard"))
 			.TabRole(ETabRole::MajorTab);
 
-		DockTab->SetContent( SAssignNew(LocalizationDashboardWidget, SLocalizationDashboard, Args.GetOwnerWindow(), DockTab, GetMutableDefault<ULocalizationTargetSet>()) );
+		DockTab->SetContent( SAssignNew(LocalizationDashboardWidget, SLocalizationDashboard, Args.GetOwnerWindow(), DockTab) );
 
 		return DockTab;
 	};
diff --git a/Engine/Source/Editor/LocalizationDashboard/Private/LocalizationDashboardModule.cpp b/Engine/Source/Editor/LocalizationDashboard/Private/LocalizationDashboardModule.cpp
--- a/Engine/Source/Editor/LocalizationDashboard/Private/LocalizationDashboardModule.cpp
+++ b/Engine/Source/Editor/LocalizationDashboard/Private/LocalizationDashboardModule.cpp
@@ -1,4 +1,4 @@
-// Copyright 1998-2015 Epic Games, Inc. All Rights Reserved.
+﻿// Copyright 1998-2015 Epic Games, Inc. All Rights Reserved.
 
 #include "LocalizationDashboardPrivatePCH.h"
 #include "ILocalizationDashboardModule.h"
@@ -16,20 +16,24 @@ class FLocalizationDashboardModule
 	: public ILocalizationDashboardModule
 {
 public:
-	FLocalizationDashboardModule()
-		: Settings(GetMutableDefault<ULocalizationTargetSet>(ULocalizationTargetSet::StaticClass()))
+	// Begin IModuleInterface interface
+	virtual void StartupModule() override
 	{
+		EngineTargetSet = NewObject<ULocalizationTargetSet>(GetTransientPackage(), ULocalizationTargetSet::EngineTargetSetName, RF_RootSet);
+		ProjectTargetSet = NewObject<ULocalizationTargetSet>(GetTransientPackage(), ULocalizationTargetSet::ProjectTargetSetName, RF_RootSet);
 		
-		for (ULocalizationTarget*& Target : Settings->TargetObjects)
+		for (ULocalizationTarget*& Target : EngineTargetSet->TargetObjects)
 		{
 			Target->UpdateStatusFromConflictReport();
 			Target->UpdateWordCountsFromCSV();
 		}
+
+		for (ULocalizationTarget*& Target : ProjectTargetSet->TargetObjects)
+		{
+			Target->UpdateStatusFromConflictReport();
+			Target->UpdateWordCountsFromCSV();
 	}
 
-	// Begin IModuleInterface interface
-	virtual void StartupModule() override
-	{
 		ServiceProviders = IModularFeatures::Get().GetModularFeatureImplementations<ILocalizationServiceProvider>("LocalizationService");
 		ServiceProviders.Insert(nullptr, 0); // "None"
 
@@ -90,7 +94,9 @@ public:
 	}
 
 private:
-	TWeakObjectPtr<ULocalizationTargetSet> Settings;
+	TWeakObjectPtr<ULocalizationTargetSet> EngineTargetSet;
+	TWeakObjectPtr<ULocalizationTargetSet> ProjectTargetSet;
+
 	TArray<ILocalizationServiceProvider*> ServiceProviders;
 	FName CurrentServiceProviderName;
 };
diff --git a/Engine/Source/Editor/LocalizationDashboard/Private/LocalizationTarget.cpp b/Engine/Source/Editor/LocalizationDashboard/Private/LocalizationTarget.cpp
--- a/Engine/Source/Editor/LocalizationDashboard/Private/LocalizationTarget.cpp
+++ b/Engine/Source/Editor/LocalizationDashboard/Private/LocalizationTarget.cpp
@@ -168,6 +168,12 @@ void ULocalizationTarget::PostEditChangeProperty(FPropertyChangedEvent& Property
 	}
 }
 
+bool ULocalizationTarget::IsMemberOfEngineTargetSet() const
+{
+	ULocalizationTargetSet* const TargetSet = CastChecked<ULocalizationTargetSet>(GetOuter());
+	return TargetSet == FindObjectChecked<ULocalizationTargetSet>(GetTransientPackage(), *ULocalizationTargetSet::EngineTargetSetName.ToString());
+}
+
 bool ULocalizationTarget::UpdateWordCountsFromCSV()
 {
 	const FString CSVFilePath = LocalizationConfigurationScript::GetWordCountCSVPath(this);
@@ -444,6 +450,9 @@ bool ULocalizationTarget::DeleteFiles(const FString* const Culture) const
 	return HasCompletelySucceeded;
 }
 
+FName ULocalizationTargetSet::EngineTargetSetName = TEXT("EngineLocalizationTargetSet");
+FName ULocalizationTargetSet::ProjectTargetSetName = TEXT("ProjectLocalizationTargetSet");
+
 void ULocalizationTargetSet::PostInitProperties()
 {
 	Super::PostInitProperties();
diff --git a/Engine/Source/Editor/UnrealEd/Private/Commandlets/CommandletHelpers.cpp b/Engine/Source/Editor/UnrealEd/Private/Commandlets/CommandletHelpers.cpp
--- a/Engine/Source/Editor/UnrealEd/Private/Commandlets/CommandletHelpers.cpp
+++ b/Engine/Source/Editor/UnrealEd/Private/Commandlets/CommandletHelpers.cpp
@@ -10,14 +10,18 @@ namespace CommandletHelpers
 		check(CommandletName && FCString::Strlen(CommandletName) != 0);
 
 		FString Arguments;
-		const FString RunCommand = FString::Printf( TEXT("-run=%s"), CommandletName );
+
 		if (ProjectPath && FCString::Strlen(ProjectPath) != 0)
 		{
-			Arguments = FString::Printf( TEXT("%s %s"), ProjectPath, *RunCommand );
+			Arguments += FString(Arguments.IsEmpty() ? TEXT("") : TEXT(" ")) + ProjectPath;
 		}
+
+		const FString RunCommand = FString::Printf( TEXT("-run=%s"), CommandletName );
+		Arguments += (Arguments.IsEmpty() ? "" : " ") + RunCommand;
+
 		if (AdditionalArguments && FCString::Strlen(AdditionalArguments) != 0)
 		{
-			Arguments = FString::Printf( TEXT("%s %s"), *Arguments, AdditionalArguments );
+			Arguments += FString(Arguments.IsEmpty() ? TEXT("") : TEXT(" ")) + AdditionalArguments;
 		}
 
 		return Arguments;
diff --git a/Engine/Source/Editor/UnrealEd/Private/Commandlets/GatherTextFromSourceCommandlet.cpp b/Engine/Source/Editor/UnrealEd/Private/Commandlets/GatherTextFromSourceCommandlet.cpp
--- a/Engine/Source/Editor/UnrealEd/Private/Commandlets/GatherTextFromSourceCommandlet.cpp
+++ b/Engine/Source/Editor/UnrealEd/Private/Commandlets/GatherTextFromSourceCommandlet.cpp
@@ -108,6 +108,8 @@ int32 UGatherTextFromSourceCommandlet::Main( const FString& Params )
 		if(FPaths::IsRelative(IncludePath))
 		{
 			IncludePath = FPaths::Combine( *ProjectBasePath, *IncludePath );
+			IncludePath = FPaths::ConvertRelativePathToFull(IncludePath);
+			FPaths::CollapseRelativeDirectories(IncludePath);
 		}
 
 		for (int32 SourceFileSearchIdx=0; SourceFileSearchIdx < UniqueSourceFileSearchFilters.Num(); SourceFileSearchIdx++)
@@ -149,7 +151,23 @@ int32 UGatherTextFromSourceCommandlet::Main( const FString& Params )
 	// Return if no source files were found
 	if( FilesToProcess.Num() == 0 )
 	{
-		UE_LOG(LogGatherTextFromSourceCommandlet, Error, TEXT("The GatherTextFromSource commandlet couldn't find any source files in the specified directories."));
+		FString SpecifiedDirectoriesString;
+		for (FString& IncludePath : IncludePaths)
+		{
+			SpecifiedDirectoriesString.Append(FString(SpecifiedDirectoriesString.IsEmpty() ? TEXT("") : TEXT("\n")) + FString::Printf(TEXT("+ %s"), *IncludePath));
+		}
+		for (FString& ExcludePath : ExcludePaths)
+		{
+			SpecifiedDirectoriesString.Append(FString(SpecifiedDirectoriesString.IsEmpty() ? TEXT("") : TEXT("\n")) + FString::Printf(TEXT("- %s"), *ExcludePath));
+		}
+
+		FString SourceFileSearchFiltersString;
+		for (const auto& Filter : UniqueSourceFileSearchFilters)
+		{
+			SourceFileSearchFiltersString += FString(SourceFileSearchFiltersString.IsEmpty() ? TEXT("") : TEXT(", ")) + Filter;
+		}
+
+		UE_LOG(LogGatherTextFromSourceCommandlet, Error, TEXT("The GatherTextFromSource commandlet couldn't find any source files matching (%s) in the specified directories:\n%s"), *SourceFileSearchFiltersString, *SpecifiedDirectoriesString);
 		return -1;
 	}
 
