Added ModulesToPreload setting to the localization dashboard UI and localization configuration script generation.

diff --git a/Engine/Source/Editor/LocalizationDashboard/Classes/LocalizationTarget.h b/Engine/Source/Editor/LocalizationDashboard/Classes/LocalizationTarget.h
--- a/Engine/Source/Editor/LocalizationDashboard/Classes/LocalizationTarget.h
+++ b/Engine/Source/Editor/LocalizationDashboard/Classes/LocalizationTarget.h
@@ -107,6 +107,12 @@ struct FLocalizationTargetSettings
 	UPROPERTY(config, EditAnywhere, Category = "Gather Configuration")
 	TArray<FString> TargetDependencies;
 
+	UPROPERTY(config, EditAnywhere, Category = "Gather Configuration", AdvancedDisplay, meta=(FilePathFilter="manifest"))
+	TArray<FFilePath> AdditionalManifestDependencies;
+
+	UPROPERTY(config, EditAnywhere, Category = "Gather Configuration", AdvancedDisplay)
+	TArray<FString> RequiredModuleNames;
+
 	UPROPERTY(config, EditAnywhere, Category = "Gather Configuration")
 	FGatherTextFromTextFilesConfiguration GatherFromTextFiles;
 
@@ -116,9 +122,6 @@ struct FLocalizationTargetSettings
 	UPROPERTY(config, EditAnywhere, Category = "Gather Configuration")
 	FGatherTextFromMetaDataConfiguration GatherFromMetaData;
 
-	UPROPERTY(config, EditAnywhere, Category = "Gather Configuration", AdvancedDisplay, meta=(FilePathFilter="manifest"))
-	TArray<FFilePath> AdditionalManifestDependencies;
-
 	UPROPERTY(config, EditAnywhere, Category = "Cultures")
 	FCultureStatistics NativeCultureStatistics;
 
diff --git a/Engine/Source/Editor/LocalizationDashboard/Private/LocalizationConfigurationScript.cpp b/Engine/Source/Editor/LocalizationDashboard/Private/LocalizationConfigurationScript.cpp
--- a/Engine/Source/Editor/LocalizationDashboard/Private/LocalizationConfigurationScript.cpp
+++ b/Engine/Source/Editor/LocalizationDashboard/Private/LocalizationConfigurationScript.cpp
@@ -108,6 +108,11 @@ namespace LocalizationConfigurationScript
 				ConfigSection.Add( TEXT("ManifestDependencies"), MakePathRelativeForCommandletProcess(Path.FilePath, !Target->IsMemberOfEngineTargetSet()) );
 			}
 
+			for (const FString& ModuleName : Target->Settings.RequiredModuleNames)
+			{
+				ConfigSection.Add( TEXT("ModulesToPreload"), ModuleName );
+			}
+
 			const FString SourcePath = ContentDirRelativeToGameDir / TEXT("Localization") / Target->Settings.Name;
 			ConfigSection.Add( TEXT("SourcePath"), SourcePath );
 			const FString DestinationPath = ContentDirRelativeToGameDir / TEXT("Localization") / Target->Settings.Name;
