Fixed localization configuration script generation code that was using the settings model for gathering from text files rather than packages when branching to generate or not generate a GatherTextFromAssetsCommandlet gather step.

diff --git a/Engine/Source/Editor/LocalizationDashboard/Private/LocalizationConfigurationScript.cpp b/Engine/Source/Editor/LocalizationDashboard/Private/LocalizationConfigurationScript.cpp
--- a/Engine/Source/Editor/LocalizationDashboard/Private/LocalizationConfigurationScript.cpp
+++ b/Engine/Source/Editor/LocalizationDashboard/Private/LocalizationConfigurationScript.cpp
@@ -159,7 +159,7 @@ namespace LocalizationConfigurationScript
 		}
 
 		// GatherTextFromAssets
-		if (Target->Settings.GatherFromTextFiles.SearchDirectories.Num() && Target->Settings.GatherFromTextFiles.FileExtensions.Num()) // Don't gather from packages if there are none specified.
+		if (Target->Settings.GatherFromPackages.IncludePathWildcards.Num() && Target->Settings.GatherFromPackages.FileExtensions.Num()) // Don't gather from packages if there are none specified.
 		{
 			FConfigSection& ConfigSection = Script.GatherTextStep(GatherTextStepIndex++);
 
