Fixed generation of localization status reports from the localization dashboard. Clicking the localization status buttons now opens the conflict report in the default editor for the file type.

diff --git a/Engine/Source/Editor/LocalizationDashboard/Classes/LocalizationTarget.h b/Engine/Source/Editor/LocalizationDashboard/Classes/LocalizationTarget.h
--- a/Engine/Source/Editor/LocalizationDashboard/Classes/LocalizationTarget.h
+++ b/Engine/Source/Editor/LocalizationDashboard/Classes/LocalizationTarget.h
@@ -6,7 +6,7 @@
 #include "LocalizationTarget.generated.h"
 
 UENUM()
-enum class ELocalizationTargetStatus : uint8
+enum class ELocalizationTargetConflictStatus : uint8
 {
 	Unknown,
 	ConflictsPresent,
@@ -102,7 +102,7 @@ struct FLocalizationTargetSettings
 	FString Name;
 
 	UPROPERTY(Transient, EditAnywhere, Category = "Target")
-	ELocalizationTargetStatus Status;
+	ELocalizationTargetConflictStatus ConflictStatus;
 
 	UPROPERTY(config, EditAnywhere, Category = "Gather Configuration")
 	TArray<FString> TargetDependencies;
diff --git a/Engine/Source/Editor/LocalizationDashboard/Private/LocalizationCommandletTasks.cpp b/Engine/Source/Editor/LocalizationDashboard/Private/LocalizationCommandletTasks.cpp
--- a/Engine/Source/Editor/LocalizationDashboard/Private/LocalizationCommandletTasks.cpp
+++ b/Engine/Source/Editor/LocalizationDashboard/Private/LocalizationCommandletTasks.cpp
@@ -22,11 +22,6 @@ bool LocalizationCommandletTasks::GatherTargets(const TSharedRef<SWindow>& Paren
 		const FString GatherScriptPath = LocalizationConfigurationScript::GetGatherScriptPath(Target);
 		LocalizationConfigurationScript::GenerateGatherScript(Target).Write(GatherScriptPath);
 		Tasks.Add(LocalizationCommandletExecution::FTask(GatherTaskName, GatherScriptPath, ShouldUseProjectFile));
-
-		const FText ReportTaskName = FText::Format(LOCTEXT("ReportTaskNameFormat", "Generate Reports for {TargetName}"), Arguments);
-		const FString ReportScriptPath = LocalizationConfigurationScript::GetReportScriptPath(Target);
-		LocalizationConfigurationScript::GenerateReportScript(Target).Write(ReportScriptPath);
-		Tasks.Add(LocalizationCommandletExecution::FTask(ReportTaskName, ReportScriptPath, ShouldUseProjectFile));
 	}
 
 	return LocalizationCommandletExecution::Execute(ParentWindow, LOCTEXT("GatherAllTargetsWindowTitle", "Gather Text for All Targets"), Tasks);
@@ -41,10 +36,6 @@ bool LocalizationCommandletTasks::GatherTarget(const TSharedRef<SWindow>& Parent
 	LocalizationConfigurationScript::GenerateGatherScript(Target).Write(GatherScriptPath);
 	Tasks.Add(LocalizationCommandletExecution::FTask(LOCTEXT("GatherTaskName", "Gather Text"), GatherScriptPath, ShouldUseProjectFile));
 
-	const FString ReportScriptPath = LocalizationConfigurationScript::GetReportScriptPath(Target);
-	LocalizationConfigurationScript::GenerateReportScript(Target).Write(ReportScriptPath);
-	Tasks.Add(LocalizationCommandletExecution::FTask(LOCTEXT("ReportTaskName", "Generate Reports"), ReportScriptPath, ShouldUseProjectFile));
-
 	FFormatNamedArguments Arguments;
 	Arguments.Add(TEXT("TargetName"), FText::FromString(Target->Settings.Name));
 
@@ -71,8 +62,8 @@ bool LocalizationCommandletTasks::ImportTargets(const TSharedRef<SWindow>& Paren
 		Tasks.Add(LocalizationCommandletExecution::FTask(ImportTaskName, ImportScriptPath, ShouldUseProjectFile));
 
 		const FText ReportTaskName = FText::Format(LOCTEXT("ReportTaskNameFormat", "Generate Reports for {TargetName}"), Arguments);
-		const FString ReportScriptPath = LocalizationConfigurationScript::GetReportScriptPath(Target);
-		LocalizationConfigurationScript::GenerateReportScript(Target).Write(ReportScriptPath);
+		const FString ReportScriptPath = LocalizationConfigurationScript::GetWordCountReportScriptPath(Target);
+		LocalizationConfigurationScript::GenerateWordCountReportScript(Target).Write(ReportScriptPath);
 		Tasks.Add(LocalizationCommandletExecution::FTask(ReportTaskName, ReportScriptPath, ShouldUseProjectFile));
 	}
 
@@ -88,8 +79,8 @@ bool LocalizationCommandletTasks::ImportTarget(const TSharedRef<SWindow>& Parent
 	LocalizationConfigurationScript::GenerateImportScript(Target, TOptional<FString>(), DirectoryPath).Write(ImportScriptPath);
 	Tasks.Add(LocalizationCommandletExecution::FTask(LOCTEXT("ImportTaskName", "Import Translations"), ImportScriptPath, ShouldUseProjectFile));
 
-	const FString ReportScriptPath = LocalizationConfigurationScript::GetReportScriptPath(Target);
-	LocalizationConfigurationScript::GenerateReportScript(Target).Write(ReportScriptPath);
+	const FString ReportScriptPath = LocalizationConfigurationScript::GetWordCountReportScriptPath(Target);
+	LocalizationConfigurationScript::GenerateWordCountReportScript(Target).Write(ReportScriptPath);
 	Tasks.Add(LocalizationCommandletExecution::FTask(LOCTEXT("ReportTaskName", "Generate Reports"), ReportScriptPath, ShouldUseProjectFile));
 
 	FFormatNamedArguments Arguments;
@@ -115,8 +106,8 @@ bool LocalizationCommandletTasks::ImportCulture(const TSharedRef<SWindow>& Paren
 	LocalizationConfigurationScript::GenerateImportScript(Target, TOptional<FString>(CultureName), FilePath).Write(ImportScriptPath);
 	Tasks.Add(LocalizationCommandletExecution::FTask(LOCTEXT("ImportTaskName", "Import Translations"), ImportScriptPath, ShouldUseProjectFile));
 
-	const FString ReportScriptPath = LocalizationConfigurationScript::GetReportScriptPath(Target);
-	LocalizationConfigurationScript::GenerateReportScript(Target).Write(ReportScriptPath);
+	const FString ReportScriptPath = LocalizationConfigurationScript::GetWordCountReportScriptPath(Target);
+	LocalizationConfigurationScript::GenerateWordCountReportScript(Target).Write(ReportScriptPath);
 	Tasks.Add(LocalizationCommandletExecution::FTask(LOCTEXT("ReportTaskName", "Generate Reports"), ReportScriptPath, ShouldUseProjectFile));
 
 	FFormatNamedArguments Arguments;
@@ -195,7 +186,7 @@ bool LocalizationCommandletTasks::ExportCulture(const TSharedRef<SWindow>& Paren
 	return HasSucceeeded;
 }
 
-bool LocalizationCommandletTasks::GenerateReportsForTargets(const TSharedRef<SWindow>& ParentWindow, const TArray<ULocalizationTarget*>& Targets)
+bool LocalizationCommandletTasks::GenerateWordCountReportsForTargets(const TSharedRef<SWindow>& ParentWindow, const TArray<ULocalizationTarget*>& Targets)
 {
 	TArray<LocalizationCommandletExecution::FTask> Tasks;
 
@@ -206,28 +197,28 @@ bool LocalizationCommandletTasks::GenerateReportsForTargets(const TSharedRef<SWi
 		FFormatNamedArguments Arguments;
 		Arguments.Add(TEXT("TargetName"), FText::FromString(Target->Settings.Name));
 
-		const FText ReportTaskName = FText::Format(LOCTEXT("ReportTaskNameFormat", "Generate Reports for {TargetName}"), Arguments);
-		const FString ReportScriptPath = LocalizationConfigurationScript::GetReportScriptPath(Target);
-		LocalizationConfigurationScript::GenerateReportScript(Target).Write(ReportScriptPath);
+		const FText ReportTaskName = FText::Format(LOCTEXT("ReportTaskNameFormat", "Generate Word Count Report for {TargetName}"), Arguments);
+		const FString ReportScriptPath = LocalizationConfigurationScript::GetWordCountReportScriptPath(Target);
+		LocalizationConfigurationScript::GenerateWordCountReportScript(Target).Write(ReportScriptPath);
 		Tasks.Add(LocalizationCommandletExecution::FTask(ReportTaskName, ReportScriptPath, ShouldUseProjectFile));
 	}
 
-	return LocalizationCommandletExecution::Execute(ParentWindow, LOCTEXT("GenerateReportsForAllTargetsWindowTitle", "Generate Reports for All Targets"), Tasks);
+	return LocalizationCommandletExecution::Execute(ParentWindow, LOCTEXT("GenerateReportsForAllTargetsWindowTitle", "Generate Word Count Reports for All Targets"), Tasks);
 }
 
-bool LocalizationCommandletTasks::GenerateReportsForTarget(const TSharedRef<SWindow>& ParentWindow, ULocalizationTarget* const Target)
+bool LocalizationCommandletTasks::GenerateWordCountReportForTarget(const TSharedRef<SWindow>& ParentWindow, ULocalizationTarget* const Target)
 {
 	TArray<LocalizationCommandletExecution::FTask> Tasks;
 	const bool ShouldUseProjectFile = !Target->IsMemberOfEngineTargetSet();
 
-	const FString ReportScriptPath = LocalizationConfigurationScript::GetReportScriptPath(Target);
-	LocalizationConfigurationScript::GenerateReportScript(Target).Write(ReportScriptPath);
-	Tasks.Add(LocalizationCommandletExecution::FTask(LOCTEXT("ReportTaskName", "Generate Reports"), ReportScriptPath, ShouldUseProjectFile));
+	const FString ReportScriptPath = LocalizationConfigurationScript::GetWordCountReportScriptPath(Target);
+	LocalizationConfigurationScript::GenerateWordCountReportScript(Target).Write(ReportScriptPath);
+	Tasks.Add(LocalizationCommandletExecution::FTask(LOCTEXT("ReportTaskName", "Generate Word Count Report"), ReportScriptPath, ShouldUseProjectFile));
 
 	FFormatNamedArguments Arguments;
 	Arguments.Add(TEXT("TargetName"), FText::FromString(Target->Settings.Name));
 
-	const FText WindowTitle = FText::Format(LOCTEXT("GenerateReportsForTargetWindowTitle", "Generate Reports for Target {TargetName}"), Arguments);
+	const FText WindowTitle = FText::Format(LOCTEXT("GenerateReportForTargetWindowTitle", "Generate Word Count Report for Target {TargetName}"), Arguments);
 	return LocalizationCommandletExecution::Execute(ParentWindow, WindowTitle, Tasks);
 }
 
diff --git a/Engine/Source/Editor/LocalizationDashboard/Private/LocalizationCommandletTasks.h b/Engine/Source/Editor/LocalizationDashboard/Private/LocalizationCommandletTasks.h
--- a/Engine/Source/Editor/LocalizationDashboard/Private/LocalizationCommandletTasks.h
+++ b/Engine/Source/Editor/LocalizationDashboard/Private/LocalizationCommandletTasks.h
@@ -16,8 +16,8 @@ namespace LocalizationCommandletTasks
 	bool ExportTarget(const TSharedRef<SWindow>& ParentWindow, ULocalizationTarget* const Target, const TOptional<FString> DirectoryPath = TOptional<FString>());
 	bool ExportCulture(const TSharedRef<SWindow>& ParentWindow, ULocalizationTarget* const Target, const FString& CultureName, const TOptional<FString> FilePath = TOptional<FString>());
 
-	bool GenerateReportsForTargets(const TSharedRef<SWindow>& ParentWindow, const TArray<ULocalizationTarget*>& Targets);
-	bool GenerateReportsForTarget(const TSharedRef<SWindow>& ParentWindow, ULocalizationTarget* const Target);
+	bool GenerateWordCountReportsForTargets(const TSharedRef<SWindow>& ParentWindow, const TArray<ULocalizationTarget*>& Targets);
+	bool GenerateWordCountReportForTarget(const TSharedRef<SWindow>& ParentWindow, ULocalizationTarget* const Target);
 
 	bool CompileTargets(const TSharedRef<SWindow>& ParentWindow, const TArray<ULocalizationTarget*>& Targets);
 	bool CompileTarget(const TSharedRef<SWindow>& ParentWindow, ULocalizationTarget* const Target);
diff --git a/Engine/Source/Editor/LocalizationDashboard/Private/LocalizationConfigurationScript.cpp b/Engine/Source/Editor/LocalizationDashboard/Private/LocalizationConfigurationScript.cpp
--- a/Engine/Source/Editor/LocalizationDashboard/Private/LocalizationConfigurationScript.cpp
+++ b/Engine/Source/Editor/LocalizationDashboard/Private/LocalizationConfigurationScript.cpp
@@ -1,4 +1,4 @@
-// Copyright 1998-2015 Epic Games, Inc. All Rights Reserved.
+﻿// Copyright 1998-2015 Epic Games, Inc. All Rights Reserved.
 
 #include "LocalizationDashboardPrivatePCH.h"
 #include "LocalizationConfigurationScript.h"
@@ -43,7 +43,7 @@ namespace LocalizationConfigurationScript
 		Result.Add(GetGatherScriptPath(Target));
 		Result.Add(GetImportScriptPath(Target));
 		Result.Add(GetExportScriptPath(Target));
-		Result.Add(GetReportScriptPath(Target));
+		Result.Add(GetWordCountReportScriptPath(Target));
 		return Result;
 	}
 
@@ -231,6 +231,34 @@ namespace LocalizationConfigurationScript
 			ConfigSection.Add( TEXT("CommandletClass"), TEXT("GenerateGatherArchive") );
 		}
 
+		// GenerateTextLocalizationReport
+		{
+			FConfigSection& ConfigSection = Script.GatherTextStep(GatherTextStepIndex++);
+
+			// CommandletClass
+			ConfigSection.Add( TEXT("CommandletClass"), TEXT("GenerateTextLocalizationReport") );
+
+			const FString SourcePath = ContentDirRelativeToGameDir / TEXT("Localization") / Target->Settings.Name;
+			ConfigSection.Add( TEXT("SourcePath"), SourcePath );
+			const FString DestinationPath = ContentDirRelativeToGameDir / TEXT("Localization") / Target->Settings.Name;
+			ConfigSection.Add( TEXT("DestinationPath"), DestinationPath );
+
+			ConfigSection.Add( TEXT("ManifestName"), FString::Printf( TEXT("%s.%s"), *Target->Settings.Name, TEXT("manifest") ) );
+
+			for (const FCultureStatistics& CultureStatistics : Target->Settings.SupportedCulturesStatistics)
+			{
+				ConfigSection.Add( TEXT("CulturesToGenerate"), CultureStatistics.CultureName );
+			}
+
+			ConfigSection.Add( TEXT("bWordCountReport"), TEXT("true") );
+
+			ConfigSection.Add( TEXT("WordCountReportName"), FPaths::GetCleanFilename( GetWordCountCSVPath(Target) ) );
+
+			ConfigSection.Add( TEXT("bConflictReport"), TEXT("true") );
+
+			ConfigSection.Add( TEXT("ConflictReportName"), FPaths::GetCleanFilename( GetConflictReportPath(Target) ) );
+		}
+
 		Script.Dirty = true;
 
 		return Script;
@@ -453,7 +481,7 @@ namespace LocalizationConfigurationScript
 		return ConfigFilePath;
 	}
 
-	FLocalizationConfigurationScript GenerateReportScript(const ULocalizationTarget* const Target)
+	FLocalizationConfigurationScript GenerateWordCountReportScript(const ULocalizationTarget* const Target)
 	{
 		FLocalizationConfigurationScript Script;
 
@@ -481,10 +509,6 @@ namespace LocalizationConfigurationScript
 			ConfigSection.Add( TEXT("bWordCountReport"), TEXT("true") );
 
 			ConfigSection.Add( TEXT("WordCountReportName"), FPaths::GetCleanFilename( GetWordCountCSVPath(Target) ) );
-
-			ConfigSection.Add( TEXT("bConflictReport"), TEXT("true") );
-
-			ConfigSection.Add( TEXT("ConflictReportName"), FPaths::GetCleanFilename( GetConflictReportPath(Target) ) );
 		}
 
 		Script.Dirty = true;
@@ -492,7 +516,7 @@ namespace LocalizationConfigurationScript
 		return Script;
 	}
 
-	FString GetReportScriptPath(const ULocalizationTarget* const Target)
+	FString GetWordCountReportScriptPath(const ULocalizationTarget* const Target)
 	{
 		return GetScriptDirectory(Target) / FString::Printf( TEXT("%s_GenerateReports.%s"), *(Target->Settings.Name), TEXT("ini") );
 	}
diff --git a/Engine/Source/Editor/LocalizationDashboard/Private/LocalizationConfigurationScript.h b/Engine/Source/Editor/LocalizationDashboard/Private/LocalizationConfigurationScript.h
--- a/Engine/Source/Editor/LocalizationDashboard/Private/LocalizationConfigurationScript.h
+++ b/Engine/Source/Editor/LocalizationDashboard/Private/LocalizationConfigurationScript.h
@@ -43,8 +43,8 @@ namespace LocalizationConfigurationScript
 	FLocalizationConfigurationScript GenerateExportScript(const ULocalizationTarget* const Target, const TOptional<FString> CultureName = TOptional<FString>(), const TOptional<FString> OutputPathOverride = TOptional<FString>());
 	FString GetExportScriptPath(const ULocalizationTarget* const Target, const TOptional<FString> CultureName = TOptional<FString>());
 
-	FLocalizationConfigurationScript GenerateReportScript(const ULocalizationTarget* const Target);
-	FString GetReportScriptPath(const ULocalizationTarget* const Target);
+	FLocalizationConfigurationScript GenerateWordCountReportScript(const ULocalizationTarget* const Target);
+	FString GetWordCountReportScriptPath(const ULocalizationTarget* const Target);
 
 	FLocalizationConfigurationScript GenerateCompileScript(const ULocalizationTarget* const Target);
 	FString GetCompileScriptPath(const ULocalizationTarget* const Target);
diff --git a/Engine/Source/Editor/LocalizationDashboard/Private/LocalizationTarget.cpp b/Engine/Source/Editor/LocalizationDashboard/Private/LocalizationTarget.cpp
--- a/Engine/Source/Editor/LocalizationDashboard/Private/LocalizationTarget.cpp
+++ b/Engine/Source/Editor/LocalizationDashboard/Private/LocalizationTarget.cpp
@@ -8,7 +8,7 @@
 #include "LocalizationConfigurationScript.h"
 
 FLocalizationTargetSettings::FLocalizationTargetSettings()
-	: Status(ELocalizationTargetStatus::Unknown)
+	: ConflictStatus(ELocalizationTargetConflictStatus::Unknown)
 {
 }
 
@@ -189,15 +189,15 @@ void ULocalizationTarget::UpdateStatusFromConflictReport()
 	{
 	case INDEX_NONE:
 		// Unknown
-		Settings.Status = ELocalizationTargetStatus::Unknown;
+		Settings.ConflictStatus = ELocalizationTargetConflictStatus::Unknown;
 		break;
 	case 0:
 		// Good
-		Settings.Status = ELocalizationTargetStatus::Clear;
+		Settings.ConflictStatus = ELocalizationTargetConflictStatus::Clear;
 		break;
 	default:
 		// Bad
-		Settings.Status = ELocalizationTargetStatus::ConflictsPresent;
+		Settings.ConflictStatus = ELocalizationTargetConflictStatus::ConflictsPresent;
 		break;
 	}
 }
@@ -342,7 +342,7 @@ bool ULocalizationTarget::RenameTargetAndFiles(const FString& NewName)
 	LocalizationConfigurationScript::GenerateGatherScript(this).Write(LocalizationConfigurationScript::GetGatherScriptPath(this));
 	LocalizationConfigurationScript::GenerateImportScript(this).Write(LocalizationConfigurationScript::GetImportScriptPath(this));
 	LocalizationConfigurationScript::GenerateExportScript(this).Write(LocalizationConfigurationScript::GetExportScriptPath(this));
-	LocalizationConfigurationScript::GenerateReportScript(this).Write(LocalizationConfigurationScript::GetReportScriptPath(this));
+	LocalizationConfigurationScript::GenerateWordCountReportScript(this).Write(LocalizationConfigurationScript::GetWordCountReportScriptPath(this));
 
 	//if (CanUseSourceControl)
 	//{
diff --git a/Engine/Source/Editor/LocalizationDashboard/Private/LocalizationTargetDetailCustomization.cpp b/Engine/Source/Editor/LocalizationDashboard/Private/LocalizationTargetDetailCustomization.cpp
--- a/Engine/Source/Editor/LocalizationDashboard/Private/LocalizationTargetDetailCustomization.cpp
+++ b/Engine/Source/Editor/LocalizationDashboard/Private/LocalizationTargetDetailCustomization.cpp
@@ -86,7 +86,7 @@ void FLocalizationTargetDetailCustomization::CustomizeDetails(IDetailLayoutBuild
 			];
 	});
 
-	PropertyCustomizationMap.Add(GET_MEMBER_NAME_CHECKED(FLocalizationTargetSettings, Status), [&](const TSharedRef<IPropertyHandle>& MemberPropertyHandle, IDetailCategoryBuilder& DetailCategoryBuilder)
+	PropertyCustomizationMap.Add(GET_MEMBER_NAME_CHECKED(FLocalizationTargetSettings, ConflictStatus), [&](const TSharedRef<IPropertyHandle>& MemberPropertyHandle, IDetailCategoryBuilder& DetailCategoryBuilder)
 	{
 		FDetailWidgetRow& StatusRow = DetailCategoryBuilder.AddCustomRow( MemberPropertyHandle->GetPropertyDisplayName() );
 		StatusRow.NameContent()
@@ -697,7 +697,7 @@ void FLocalizationTargetDetailCustomization::RefreshWordCounts()
 	if (LocalizationTarget.IsValid())
 	{
 		const TSharedPtr<SWindow> ParentWindow = FSlateApplication::Get().FindWidgetWindow(DetailLayoutBuilder->GetDetailsView().AsShared());
-		LocalizationCommandletTasks::GenerateReportsForTarget(ParentWindow.ToSharedRef(), LocalizationTarget.Get());
+		LocalizationCommandletTasks::GenerateWordCountReportForTarget(ParentWindow.ToSharedRef(), LocalizationTarget.Get());
 
 		UpdateTargetFromReports();
 	}
diff --git a/Engine/Source/Editor/LocalizationDashboard/Private/LocalizationTargetSetDetailCustomization.cpp b/Engine/Source/Editor/LocalizationDashboard/Private/LocalizationTargetSetDetailCustomization.cpp
--- a/Engine/Source/Editor/LocalizationDashboard/Private/LocalizationTargetSetDetailCustomization.cpp
+++ b/Engine/Source/Editor/LocalizationDashboard/Private/LocalizationTargetSetDetailCustomization.cpp
@@ -134,7 +134,7 @@ void FLocalizationTargetSetDetailCustomization::CustomizeDetails(IDetailLayoutBu
 							.HAlignCell(HAlign_Left)
 							.VAlignCell(VAlign_Center)
 							+SHeaderRow::Column("Status")
-							.DefaultLabel( LOCTEXT("StatusColumnLabel", "Status"))
+							.DefaultLabel( LOCTEXT("StatusColumnLabel", "Conflict Status"))
 							.HAlignHeader(HAlign_Center)
 							.HAlignCell(HAlign_Center)
 							.VAlignCell(VAlign_Center)
diff --git a/Engine/Source/Editor/LocalizationDashboard/Private/SLocalizationDashboardTargetRow.cpp b/Engine/Source/Editor/LocalizationDashboard/Private/SLocalizationDashboardTargetRow.cpp
--- a/Engine/Source/Editor/LocalizationDashboard/Private/SLocalizationDashboardTargetRow.cpp
+++ b/Engine/Source/Editor/LocalizationDashboard/Private/SLocalizationDashboardTargetRow.cpp
@@ -314,7 +314,7 @@ FReply SLocalizationDashboardTargetRow::RefreshWordCount()
 	if (LocalizationTarget)
 	{
 		const TSharedPtr<SWindow> ParentWindow = FSlateApplication::Get().FindWidgetWindow(AsShared());
-		LocalizationCommandletTasks::GenerateReportsForTarget(ParentWindow.ToSharedRef(), LocalizationTarget);
+		LocalizationCommandletTasks::GenerateWordCountReportForTarget(ParentWindow.ToSharedRef(), LocalizationTarget);
 
 		UpdateTargetFromReports();
 	}
diff --git a/Engine/Source/Editor/LocalizationDashboard/Private/SLocalizationTargetStatusButton.cpp b/Engine/Source/Editor/LocalizationDashboard/Private/SLocalizationTargetStatusButton.cpp
--- a/Engine/Source/Editor/LocalizationDashboard/Private/SLocalizationTargetStatusButton.cpp
+++ b/Engine/Source/Editor/LocalizationDashboard/Private/SLocalizationTargetStatusButton.cpp
@@ -30,16 +30,16 @@ void SLocalizationTargetStatusButton::Construct(const FArguments& InArgs, ULocal
 
 const FSlateBrush* SLocalizationTargetStatusButton::GetImageBrush() const
 {
-	switch (Target->Settings.Status)
+	switch (Target->Settings.ConflictStatus)
 	{
 	default:
-	case ELocalizationTargetStatus::Unknown:
+	case ELocalizationTargetConflictStatus::Unknown:
 		return FEditorStyle::GetBrush("Icons.Warning");
 		break;
-	case ELocalizationTargetStatus::Clear:
+	case ELocalizationTargetConflictStatus::Clear:
 		return FEditorStyle::GetBrush("Symbols.Check");
 		break;
-	case ELocalizationTargetStatus::ConflictsPresent:
+	case ELocalizationTargetConflictStatus::ConflictsPresent:
 		return FEditorStyle::GetBrush("Symbols.X");
 		break;
 	}
@@ -47,16 +47,16 @@ const FSlateBrush* SLocalizationTargetStatusButton::GetImageBrush() const
 
 FSlateColor SLocalizationTargetStatusButton::GetColorAndOpacity() const
 {
-	switch (Target->Settings.Status)
+	switch (Target->Settings.ConflictStatus)
 	{
 	default:
-	case ELocalizationTargetStatus::Unknown:
+	case ELocalizationTargetConflictStatus::Unknown:
 		return FLinearColor::White;
 		break;
-	case ELocalizationTargetStatus::Clear:
+	case ELocalizationTargetConflictStatus::Clear:
 		return FLinearColor::Green;
 		break;
-	case ELocalizationTargetStatus::ConflictsPresent:
+	case ELocalizationTargetConflictStatus::ConflictsPresent:
 		return FLinearColor::Red;
 		break;
 	}
@@ -64,16 +64,16 @@ FSlateColor SLocalizationTargetStatusButton::GetColorAndOpacity() const
 
 FText SLocalizationTargetStatusButton::GetToolTipText() const
 {
-	switch (Target->Settings.Status)
+	switch (Target->Settings.ConflictStatus)
 	{
 	default:
-	case ELocalizationTargetStatus::Unknown:
-		return LOCTEXT("StatusToolTip_Unknown", "Conflict report file not detected. Click to generate the conflict report.");
+	case ELocalizationTargetConflictStatus::Unknown:
+		return LOCTEXT("StatusToolTip_Unknown", "Conflict report file not detected. Perform a gather to generate a conflict report file.");
 		break;
-	case ELocalizationTargetStatus::Clear:
+	case ELocalizationTargetConflictStatus::Clear:
 		return LOCTEXT("StatusToolTip_Clear", "No conflicts detected.");
 		break;
-	case ELocalizationTargetStatus::ConflictsPresent:
+	case ELocalizationTargetConflictStatus::ConflictsPresent:
 		return LOCTEXT("StatusToolTip_ConflictsPresent", "Conflicts detected. Click to open the conflict report.");
 		break;
 	}
@@ -81,42 +81,16 @@ FText SLocalizationTargetStatusButton::GetToolTipText() const
 
 FReply SLocalizationTargetStatusButton::OnClicked()
 {
-	switch (Target->Settings.Status)
+	switch (Target->Settings.ConflictStatus)
 	{
-	case ELocalizationTargetStatus::Unknown:
-		{
-			// Generate conflict report.
-			const TSharedPtr<SWindow> ParentWindow = FSlateApplication::Get().FindWidgetWindow(AsShared());
-			if (LocalizationCommandletTasks::GenerateReportsForTarget(ParentWindow.ToSharedRef(), Target))
-			{
-				Target->UpdateStatusFromConflictReport();
-			}
-		}
+	case ELocalizationTargetConflictStatus::Unknown:
+		// Do nothing.
 		break;
-	case ELocalizationTargetStatus::Clear:
+	case ELocalizationTargetConflictStatus::Clear:
 		// Do nothing.
 		break;
-	case ELocalizationTargetStatus::ConflictsPresent:
-		FString ReportString;
-		if (FFileHelper::LoadFileToString(ReportString, *LocalizationConfigurationScript::GetConflictReportPath(Target)))
-		{
-			TSharedPtr<SWindow> ReportWindow =
-				SNew(SWindow)
-				.Content()
-				[
-					SNew(SBorder)
-					.BorderImage(FCoreStyle::Get().GetBrush("ToolPanel.GroupBorder"))
-					[
-						SNew(SMultiLineEditableText)
-						.Text( FText::FromString(ReportString) )
-					]
-				];
-
-			if (ReportWindow.IsValid())
-			{
-				FSlateApplication::Get().AddWindow(ReportWindow.ToSharedRef());
-			}
-		}
+	case ELocalizationTargetConflictStatus::ConflictsPresent:
+		FPlatformProcess::LaunchFileInDefaultExternalApplication(*FPaths::ConvertRelativePathToFull(LocalizationConfigurationScript::GetConflictReportPath(Target)), nullptr, ELaunchVerb::Open);
 		break;
 	}
 
