Localization targets may now depend on engine localization targets.

diff --git a/Engine/Source/Editor/LocalizationDashboard/Private/LocalizationConfigurationScript.cpp b/Engine/Source/Editor/LocalizationDashboard/Private/LocalizationConfigurationScript.cpp
--- a/Engine/Source/Editor/LocalizationDashboard/Private/LocalizationConfigurationScript.cpp
+++ b/Engine/Source/Editor/LocalizationDashboard/Private/LocalizationConfigurationScript.cpp
@@ -96,7 +96,15 @@ namespace LocalizationConfigurationScript
 			const ULocalizationTargetSet* const LocalizationTargetSet = GetDefault<ULocalizationTargetSet>(ULocalizationTargetSet::StaticClass());
 			for (const FString& TargetDependencyName : Target->Settings.TargetDependencies)
 			{
-				ULocalizationTarget* const * OtherTarget = LocalizationTargetSet->TargetObjects.FindByPredicate([&TargetDependencyName](ULocalizationTarget* const OtherTarget)->bool{return OtherTarget->Settings.Name == TargetDependencyName;});
+				TArray<ULocalizationTarget*> AllLocalizationTargets;
+				ULocalizationTargetSet* EngineTargetSet = FindObjectChecked<ULocalizationTargetSet>(ANY_PACKAGE, *ULocalizationTargetSet::EngineTargetSetName.ToString());
+				if (EngineTargetSet != LocalizationTargetSet)
+				{
+					AllLocalizationTargets.Append(EngineTargetSet->TargetObjects);
+				}
+				AllLocalizationTargets.Append(LocalizationTargetSet->TargetObjects);
+
+				ULocalizationTarget* const * OtherTarget = AllLocalizationTargets.FindByPredicate([&TargetDependencyName](ULocalizationTarget* const OtherTarget)->bool{return OtherTarget->Settings.Name == TargetDependencyName;});
 				if (OtherTarget)
 				{
 					ConfigSection.Add( TEXT("ManifestDependencies"), MakePathRelativeForCommandletProcess(GetManifestPath(*OtherTarget), !Target->IsMemberOfEngineTargetSet()) );
diff --git a/Engine/Source/Editor/LocalizationDashboard/Private/LocalizationTargetDetailCustomization.cpp b/Engine/Source/Editor/LocalizationDashboard/Private/LocalizationTargetDetailCustomization.cpp
--- a/Engine/Source/Editor/LocalizationDashboard/Private/LocalizationTargetDetailCustomization.cpp
+++ b/Engine/Source/Editor/LocalizationDashboard/Private/LocalizationTargetDetailCustomization.cpp
@@ -464,7 +464,14 @@ void FLocalizationTargetDetailCustomization::RebuildTargetsList()
 		return false;
 	};
 
-	for (ULocalizationTarget* const OtherTarget : TargetSet->TargetObjects)
+	TArray<ULocalizationTarget*> AllLocalizationTargets;
+	ULocalizationTargetSet* EngineTargetSet = FindObjectChecked<ULocalizationTargetSet>(ANY_PACKAGE, *ULocalizationTargetSet::EngineTargetSetName.ToString());
+	if (EngineTargetSet != TargetSet)
+	{
+		AllLocalizationTargets.Append(EngineTargetSet->TargetObjects);
+	}
+	AllLocalizationTargets.Append(TargetSet->TargetObjects);
+	for (ULocalizationTarget* const OtherTarget : AllLocalizationTargets)
 	{
 		if (OtherTarget != LocalizationTarget)
 		{
