Made sure that localization target names must be unique among all localization targets, not just those in the same localization target set.

diff --git a/Engine/Source/Editor/LocalizationDashboard/Private/LocalizationTargetDetailCustomization.cpp b/Engine/Source/Editor/LocalizationDashboard/Private/LocalizationTargetDetailCustomization.cpp
--- a/Engine/Source/Editor/LocalizationDashboard/Private/LocalizationTargetDetailCustomization.cpp
+++ b/Engine/Source/Editor/LocalizationDashboard/Private/LocalizationTargetDetailCustomization.cpp
@@ -341,7 +341,15 @@ FText FLocalizationTargetDetailCustomization::GetTargetName() const
 
 bool FLocalizationTargetDetailCustomization::IsTargetNameUnique(const FString& Name) const
 {
-	for (ULocalizationTarget* const TargetObject : TargetSet->TargetObjects)
+	TArray<ULocalizationTarget*> AllLocalizationTargets;
+	ULocalizationTargetSet* EngineTargetSet = FindObjectChecked<ULocalizationTargetSet>(ANY_PACKAGE, *ULocalizationTargetSet::EngineTargetSetName.ToString());
+	if (EngineTargetSet != TargetSet)
+	{
+		AllLocalizationTargets.Append(EngineTargetSet->TargetObjects);
+	}
+	AllLocalizationTargets.Append(TargetSet->TargetObjects);
+
+	for (ULocalizationTarget* const TargetObject : AllLocalizationTargets)
 	{
 		if (TargetObject != LocalizationTarget)
 		{
