Added localization target loading policy to localization dashboard UI. Allows localization targets to be added to/removed from the appropriate localization paths settings.

diff --git a/Engine/Source/Editor/LocalizationDashboard/Classes/LocalizationTarget.h b/Engine/Source/Editor/LocalizationDashboard/Classes/LocalizationTarget.h
--- a/Engine/Source/Editor/LocalizationDashboard/Classes/LocalizationTarget.h
+++ b/Engine/Source/Editor/LocalizationDashboard/Classes/LocalizationTarget.h
@@ -90,6 +90,17 @@ struct FCultureStatistics
 	uint32 WordCount;
 };
 
+UENUM()
+enum class ELocalizationTargetLoadingPolicy : uint8
+{
+	Never,
+	Always,
+	Editor,
+	Game,
+	PropertyNames,
+	ToolTips,
+};
+
 USTRUCT()
 struct FLocalizationTargetSettings
 {
diff --git a/Engine/Source/Editor/LocalizationDashboard/Private/LocalizationTargetDetailCustomization.cpp b/Engine/Source/Editor/LocalizationDashboard/Private/LocalizationTargetDetailCustomization.cpp
--- a/Engine/Source/Editor/LocalizationDashboard/Private/LocalizationTargetDetailCustomization.cpp
+++ b/Engine/Source/Editor/LocalizationDashboard/Private/LocalizationTargetDetailCustomization.cpp
@@ -16,6 +16,28 @@
 
 #define LOCTEXT_NAMESPACE "LocalizationTargetEditor"
 
+namespace
+{
+	struct FConfigValueIdentity
+	{
+		FString SectionName;
+		FString KeyName;
+		FString ConfigFilePath;
+	};
+
+	TMap<ELocalizationTargetLoadingPolicy, FConfigValueIdentity> LoadingPolicyToConfigSettingMap = []()
+	{
+		TMap<ELocalizationTargetLoadingPolicy, FConfigValueIdentity> Map;
+		Map.Add(ELocalizationTargetLoadingPolicy::Always,			FConfigValueIdentity {TEXT("Internationalization"),	TEXT("LocalizationPaths"),				GEngineIni });
+		Map.Add(ELocalizationTargetLoadingPolicy::Editor,			FConfigValueIdentity {TEXT("Internationalization"),	TEXT("LocalizationPaths"),				GEditorIni });
+		Map.Add(ELocalizationTargetLoadingPolicy::Editor,			FConfigValueIdentity {TEXT("Internationalization"),	TEXT("LocalizationPaths"),				GEditorIni });
+		Map.Add(ELocalizationTargetLoadingPolicy::PropertyNames,	FConfigValueIdentity {TEXT("Internationalization"),	TEXT("PropertyNameLocalizationPaths"),	GEditorIni });
+		Map.Add(ELocalizationTargetLoadingPolicy::ToolTips,			FConfigValueIdentity {TEXT("Internationalization"),	TEXT("ToolTipLocalizationPaths"),		GEditorIni });
+		Map.Add(ELocalizationTargetLoadingPolicy::Game,				FConfigValueIdentity {TEXT("Internationalization"),	TEXT("LocalizationPaths"),				GGameIni });
+		return Map;
+	}();
+}
+
 FLocalizationTargetDetailCustomization::FLocalizationTargetDetailCustomization()
 	: DetailLayoutBuilder(nullptr)
 	, NewEntryIndexToBeInitialized(INDEX_NONE)
@@ -309,6 +331,46 @@ void FLocalizationTargetDetailCustomization::CustomizeDetails(IDetailLayoutBuild
 		}
 	}
 
+	{
+		IDetailCategoryBuilder& DetailCategoryBuilder = DetailBuilder.EditCategory("Target");
+		FDetailWidgetRow& DetailWidgetRow = DetailCategoryBuilder.AddCustomRow(LOCTEXT("LocalizationTargetLoadingPolicyRowFilterString", "Loading Policy"));
+
+		static const TArray< TSharedPtr<ELocalizationTargetLoadingPolicy> > LoadingPolicies = []()
+		{
+			UEnum* const LoadingPolicyEnum = FindObjectChecked<UEnum>(ANY_PACKAGE, TEXT("ELocalizationTargetLoadingPolicy"));
+			TArray< TSharedPtr<ELocalizationTargetLoadingPolicy> > Array;
+			for (int32 i = 0; i < LoadingPolicyEnum->NumEnums(); ++i)
+			{
+				Array.Add( MakeShareable( new ELocalizationTargetLoadingPolicy(static_cast<ELocalizationTargetLoadingPolicy>(i)) ) );
+			}
+			return Array;
+		}();
+
+		DetailWidgetRow.NameContent()
+			[
+				SNew(STextBlock)
+				.Text(LOCTEXT("LocalizationTargetLoadingPolicyRowName", "Loading Policy"))
+			];
+		DetailWidgetRow.ValueContent()
+			[
+				SNew(SComboBox< TSharedPtr<ELocalizationTargetLoadingPolicy> >)
+				.OptionsSource(&LoadingPolicies)
+				.OnSelectionChanged(this, &FLocalizationTargetDetailCustomization::OnLoadingPolicySelectionChanged)
+				.OnGenerateWidget(this, &FLocalizationTargetDetailCustomization::GenerateWidgetForLoadingPolicy)
+				.InitiallySelectedItem(LoadingPolicies[static_cast<int32>(GetLoadingPolicy())])
+				.Content()
+				[
+					SNew(STextBlock)
+					.Text_Lambda([this]() 
+					{
+						UEnum* const LoadingPolicyEnum = FindObjectChecked<UEnum>(ANY_PACKAGE, TEXT("ELocalizationTargetLoadingPolicy"));
+						return LoadingPolicyEnum->GetDisplayNameText(static_cast<int32>(GetLoadingPolicy()));
+					})
+				]
+			];
+	}
+
+
 	//{
 	//	IDetailCategoryBuilder& ServiceProviderCategoryBuilder = DetailBuilder.EditCategory(FName("LocalizationServiceProvider"));
 
@@ -422,6 +484,59 @@ void FLocalizationTargetDetailCustomization::OnTargetNameCommitted(const FText&
 	}
 }
 
+ELocalizationTargetLoadingPolicy FLocalizationTargetDetailCustomization::GetLoadingPolicy() const
+{
+	const FString DataDirectory = LocalizationConfigurationScript::GetDataDirectory(LocalizationTarget.Get());
+
+	for (const TPair<ELocalizationTargetLoadingPolicy, FConfigValueIdentity>& Pair : LoadingPolicyToConfigSettingMap)
+	{
+		TArray<FString> LocalizationPaths;
+		GConfig->GetArray( *Pair.Value.SectionName, *Pair.Value.KeyName, LocalizationPaths, Pair.Value.ConfigFilePath );
+
+		if (LocalizationPaths.Contains(DataDirectory))
+		{
+			return Pair.Key;
+		}
+	}
+
+	return ELocalizationTargetLoadingPolicy::Never;
+}
+
+void FLocalizationTargetDetailCustomization::SetLoadingPolicy(const ELocalizationTargetLoadingPolicy LoadingPolicy)
+{
+	const FString DataDirectory = LocalizationConfigurationScript::GetDataDirectory(LocalizationTarget.Get());
+
+	for (const TPair<ELocalizationTargetLoadingPolicy, FConfigValueIdentity>& Pair : LoadingPolicyToConfigSettingMap)
+	{
+		TArray<FString> LocalizationPaths;
+		GConfig->GetArray( *Pair.Value.SectionName, *Pair.Value.KeyName, LocalizationPaths, Pair.Value.ConfigFilePath );
+		// Add this localization target's data directory to the appropriate localization path setting.
+		if (Pair.Key == LoadingPolicy && !LocalizationPaths.Contains(DataDirectory))
+		{
+			LocalizationPaths.Add(DataDirectory);
+			GConfig->SetArray( *Pair.Value.SectionName, *Pair.Value.KeyName, LocalizationPaths, Pair.Value.ConfigFilePath );
+		}
+		// Remove this localization target's data directory from any inappropriate localization path setting.
+		else if (LocalizationPaths.Contains(DataDirectory))
+		{
+			LocalizationPaths.Remove(DataDirectory);
+			GConfig->SetArray( *Pair.Value.SectionName, *Pair.Value.KeyName, LocalizationPaths, Pair.Value.ConfigFilePath );
+		}
+	}
+}
+
+void FLocalizationTargetDetailCustomization::OnLoadingPolicySelectionChanged(TSharedPtr<ELocalizationTargetLoadingPolicy> LoadingPolicy, ESelectInfo::Type SelectInfo)
+{
+	SetLoadingPolicy(*LoadingPolicy.Get());
+};
+
+TSharedRef<SWidget> FLocalizationTargetDetailCustomization::GenerateWidgetForLoadingPolicy(TSharedPtr<ELocalizationTargetLoadingPolicy> LoadingPolicy)
+{
+	UEnum* const LoadingPolicyEnum = FindObjectChecked<UEnum>(ANY_PACKAGE, TEXT("ELocalizationTargetLoadingPolicy"));
+	return SNew(STextBlock)
+		.Text(LoadingPolicyEnum->GetDisplayNameText(static_cast<int32>(*LoadingPolicy.Get())));
+};
+
 void FLocalizationTargetDetailCustomization::RebuildTargetDependenciesBox()
 {
 	if (TargetDependenciesHorizontalBox.IsValid())
diff --git a/Engine/Source/Editor/LocalizationDashboard/Private/LocalizationTargetDetailCustomization.h b/Engine/Source/Editor/LocalizationDashboard/Private/LocalizationTargetDetailCustomization.h
--- a/Engine/Source/Editor/LocalizationDashboard/Private/LocalizationTargetDetailCustomization.h
+++ b/Engine/Source/Editor/LocalizationDashboard/Private/LocalizationTargetDetailCustomization.h
@@ -22,6 +22,11 @@ private:
 	void OnTargetNameChanged(const FText& NewText);
 	void OnTargetNameCommitted(const FText& NewText, ETextCommit::Type Type);
 
+	ELocalizationTargetLoadingPolicy GetLoadingPolicy() const;
+	void SetLoadingPolicy(const ELocalizationTargetLoadingPolicy LoadingPolicy);
+	void OnLoadingPolicySelectionChanged(TSharedPtr<ELocalizationTargetLoadingPolicy> LoadingPolicy, ESelectInfo::Type SelectInfo);
+	TSharedRef<SWidget> GenerateWidgetForLoadingPolicy(TSharedPtr<ELocalizationTargetLoadingPolicy> LoadingPolicy);
+
 	void RebuildTargetDependenciesBox();
 	void RebuildTargetsList();
 	TSharedRef<ITableRow> OnGenerateTargetRow(ULocalizationTarget* OtherLocalizationTarget, const TSharedRef<STableViewBase>& Table);
