Added GUIDs to localization target so that renaming a localization target does not break references between localization targets.

diff --git a/Engine/Source/Editor/LocalizationDashboard/Classes/LocalizationTarget.h b/Engine/Source/Editor/LocalizationDashboard/Classes/LocalizationTarget.h
--- a/Engine/Source/Editor/LocalizationDashboard/Classes/LocalizationTarget.h
+++ b/Engine/Source/Editor/LocalizationDashboard/Classes/LocalizationTarget.h
@@ -112,11 +112,14 @@ struct FLocalizationTargetSettings
 	UPROPERTY(config, EditAnywhere, Category = "Target")
 	FString Name;
 
+	UPROPERTY(config)
+	FGuid Guid;
+
 	UPROPERTY(Transient, EditAnywhere, Category = "Target")
 	ELocalizationTargetConflictStatus ConflictStatus;
 
 	UPROPERTY(config, EditAnywhere, Category = "Gather Configuration")
-	TArray<FString> TargetDependencies;
+	TArray<FGuid> TargetDependencies;
 
 	UPROPERTY(config, EditAnywhere, Category = "Gather Configuration", AdvancedDisplay, meta=(FilePathFilter="manifest"))
 	TArray<FFilePath> AdditionalManifestDependencies;
diff --git a/Engine/Source/Editor/LocalizationDashboard/Private/LocalizationConfigurationScript.cpp b/Engine/Source/Editor/LocalizationDashboard/Private/LocalizationConfigurationScript.cpp
--- a/Engine/Source/Editor/LocalizationDashboard/Private/LocalizationConfigurationScript.cpp
+++ b/Engine/Source/Editor/LocalizationDashboard/Private/LocalizationConfigurationScript.cpp
@@ -94,7 +94,7 @@ namespace LocalizationConfigurationScript
 			FConfigSection& ConfigSection = Script.CommonSettings();
 
 			const ULocalizationTargetSet* const LocalizationTargetSet = GetDefault<ULocalizationTargetSet>(ULocalizationTargetSet::StaticClass());
-			for (const FString& TargetDependencyName : Target->Settings.TargetDependencies)
+			for (const FGuid& TargetDependencyGuid : Target->Settings.TargetDependencies)
 			{
 				TArray<ULocalizationTarget*> AllLocalizationTargets;
 				ULocalizationTargetSet* EngineTargetSet = FindObjectChecked<ULocalizationTargetSet>(ANY_PACKAGE, *ULocalizationTargetSet::EngineTargetSetName.ToString());
@@ -104,7 +104,7 @@ namespace LocalizationConfigurationScript
 				}
 				AllLocalizationTargets.Append(LocalizationTargetSet->TargetObjects);
 
-				ULocalizationTarget* const * OtherTarget = AllLocalizationTargets.FindByPredicate([&TargetDependencyName](ULocalizationTarget* const OtherTarget)->bool{return OtherTarget->Settings.Name == TargetDependencyName;});
+				ULocalizationTarget* const * OtherTarget = AllLocalizationTargets.FindByPredicate([&TargetDependencyGuid](ULocalizationTarget* const OtherTarget)->bool{return OtherTarget->Settings.Guid == TargetDependencyGuid;});
 				if (OtherTarget)
 				{
 					ConfigSection.Add( TEXT("ManifestDependencies"), MakePathRelativeForCommandletProcess(GetManifestPath(*OtherTarget), !Target->IsMemberOfEngineTargetSet()) );
diff --git a/Engine/Source/Editor/LocalizationDashboard/Private/LocalizationTarget.cpp b/Engine/Source/Editor/LocalizationDashboard/Private/LocalizationTarget.cpp
--- a/Engine/Source/Editor/LocalizationDashboard/Private/LocalizationTarget.cpp
+++ b/Engine/Source/Editor/LocalizationDashboard/Private/LocalizationTarget.cpp
@@ -8,7 +8,8 @@
 #include "LocalizationConfigurationScript.h"
 
 FLocalizationTargetSettings::FLocalizationTargetSettings()
-	: ConflictStatus(ELocalizationTargetConflictStatus::Unknown)
+	: Guid(FGuid::NewGuid())
+	, ConflictStatus(ELocalizationTargetConflictStatus::Unknown)
 {
 }
 
diff --git a/Engine/Source/Editor/LocalizationDashboard/Private/LocalizationTargetDetailCustomization.cpp b/Engine/Source/Editor/LocalizationDashboard/Private/LocalizationTargetDetailCustomization.cpp
--- a/Engine/Source/Editor/LocalizationDashboard/Private/LocalizationTargetDetailCustomization.cpp
+++ b/Engine/Source/Editor/LocalizationDashboard/Private/LocalizationTargetDetailCustomization.cpp
@@ -547,19 +547,33 @@ void FLocalizationTargetDetailCustomization::RebuildTargetDependenciesBox()
 		}
 		TargetDependenciesWidgets.Empty();
 
-		for (const FString& TargetDependency : LocalizationTarget->Settings.TargetDependencies)
+		TArray<ULocalizationTarget*> AllLocalizationTargets;
+		ULocalizationTargetSet* EngineTargetSet = FindObjectChecked<ULocalizationTargetSet>(ANY_PACKAGE, *ULocalizationTargetSet::EngineTargetSetName.ToString());
+		if (EngineTargetSet != TargetSet)
 		{
-			const TSharedRef<SWidget> Widget = SNew(SBorder)
+			AllLocalizationTargets.Append(EngineTargetSet->TargetObjects);
+		}
+		AllLocalizationTargets.Append(TargetSet->TargetObjects);
+
+		for (const FGuid& TargetDependencyGuid : LocalizationTarget->Settings.TargetDependencies)
+		{
+			ULocalizationTarget** const TargetDependency = AllLocalizationTargets.FindByPredicate([TargetDependencyGuid](ULocalizationTarget* SomeLocalizationTarget)->bool{return SomeLocalizationTarget->Settings.Guid == TargetDependencyGuid;});
+			if (TargetDependency)
+			{
+				TWeakObjectPtr<ULocalizationTarget> TargetDependencyPtr = *TargetDependency;
+				const auto& GetTargetDependencyName = [TargetDependencyPtr] {return FText::FromString(TargetDependencyPtr->Settings.Name);};
+				const TSharedRef<SWidget> Widget = SNew(SBorder)
 				[
 					SNew(STextBlock)
-					.Text(FText::FromString(TargetDependency))
+					.Text_Lambda(GetTargetDependencyName)
 				];
 
-			TargetDependenciesWidgets.Add(Widget);
-			TargetDependenciesHorizontalBox->AddSlot()
+				TargetDependenciesWidgets.Add(Widget);
+				TargetDependenciesHorizontalBox->AddSlot()
 				[
 					Widget
 				];
+			}
 		}
 	}
 }
@@ -570,14 +584,14 @@ void FLocalizationTargetDetailCustomization::RebuildTargetsList()
 	TFunction<bool (ULocalizationTarget* const)> DoesTargetDependOnUs;
 	DoesTargetDependOnUs = [&, this](ULocalizationTarget* const OtherTarget) -> bool
 	{
-		if (OtherTarget->Settings.TargetDependencies.Contains(LocalizationTarget->Settings.Name))
+		if (OtherTarget->Settings.TargetDependencies.Contains(LocalizationTarget->Settings.Guid))
 		{
 			return true;
 		}
 
-		for (const FString& OtherTargetDependencyName : OtherTarget->Settings.TargetDependencies)
+		for (const FGuid& OtherTargetDependencyGuid : OtherTarget->Settings.TargetDependencies)
 		{
-			ULocalizationTarget** const OtherTargetDependency = TargetSet->TargetObjects.FindByPredicate([OtherTargetDependencyName](ULocalizationTarget* SomeLocalizationTarget)->bool{return SomeLocalizationTarget->Settings.Name == OtherTargetDependencyName;});
+			ULocalizationTarget** const OtherTargetDependency = TargetSet->TargetObjects.FindByPredicate([OtherTargetDependencyGuid](ULocalizationTarget* SomeLocalizationTarget)->bool{return SomeLocalizationTarget->Settings.Guid == OtherTargetDependencyGuid;});
 			if (OtherTargetDependency && DoesTargetDependOnUs(*OtherTargetDependency))
 			{
 				return true;
@@ -646,10 +660,10 @@ void FLocalizationTargetDetailCustomization::OnTargetDependencyCheckStateChanged
 	switch (State)
 	{
 	case ECheckBoxState::Checked:
-		LocalizationTarget->Settings.TargetDependencies.Add(OtherLocalizationTarget->Settings.Name);
+		LocalizationTarget->Settings.TargetDependencies.Add(OtherLocalizationTarget->Settings.Guid);
 		break;
 	case ECheckBoxState::Unchecked:
-		LocalizationTarget->Settings.TargetDependencies.Remove(OtherLocalizationTarget->Settings.Name);
+		LocalizationTarget->Settings.TargetDependencies.Remove(OtherLocalizationTarget->Settings.Guid);
 		break;
 	}
 
@@ -663,7 +677,7 @@ void FLocalizationTargetDetailCustomization::OnTargetDependencyCheckStateChanged
 
 ECheckBoxState FLocalizationTargetDetailCustomization::IsTargetDependencyChecked(ULocalizationTarget* const OtherLocalizationTarget) const
 {
-	return LocalizationTarget->Settings.TargetDependencies.Contains(OtherLocalizationTarget->Settings.Name) ? ECheckBoxState::Checked : ECheckBoxState::Unchecked;
+	return LocalizationTarget->Settings.TargetDependencies.Contains(OtherLocalizationTarget->Settings.Guid) ? ECheckBoxState::Checked : ECheckBoxState::Unchecked;
 }
 
 FText FLocalizationTargetDetailCustomization::GetNativeCultureName() const
