Added descriptive comments to localization target types and members for use in the localization dashboard UI.

diff --git a/Engine/Source/Editor/LocalizationDashboard/Classes/LocalizationTarget.h b/Engine/Source/Editor/LocalizationDashboard/Classes/LocalizationTarget.h
--- a/Engine/Source/Editor/LocalizationDashboard/Classes/LocalizationTarget.h
+++ b/Engine/Source/Editor/LocalizationDashboard/Classes/LocalizationTarget.h
@@ -1,4 +1,4 @@
-// Copyright 1998-2015 Epic Games, Inc. All Rights Reserved.
+﻿// Copyright 1998-2015 Epic Games, Inc. All Rights Reserved.
 
 #pragma once
 
@@ -8,9 +8,9 @@
 UENUM()
 enum class ELocalizationTargetConflictStatus : uint8
 {
-	Unknown,
-	ConflictsPresent,
-	Clear
+	Unknown, /* The status of conflicts in this localization target could not be determined. */
+	ConflictsPresent, /* The are outstanding conflicts present in this localization target. */
+	Clear, /* The localization target is clear of conflicts. */
 };
 
 USTRUCT()
@@ -18,12 +18,15 @@ struct FGatherTextFromTextFilesConfiguration
 {
 	GENERATED_USTRUCT_BODY()
 
+	/* The paths of directories to be searched for text files, specified relative to the project's root, which may be parsed for text to gather. */
 	UPROPERTY(config, EditAnywhere, Category = "Filter")
 	TArray<FString> SearchDirectories;
 
+	/* Text files whose paths match these wildcard patterns will be excluded from gathering. */
 	UPROPERTY(config, EditAnywhere, Category = "Filter")
 	TArray<FString> ExcludePathWildcards;
 
+	/* Text files whose names match these wildcard patterns may be parsed for text to gather. */
 	UPROPERTY(config, EditAnywhere, Category = "Filter")
 	TArray<FString> FileExtensions;
 };
@@ -33,12 +36,15 @@ struct FGatherTextFromPackagesConfiguration
 {
 	GENERATED_USTRUCT_BODY()
 
+	/* Packages whose paths match these wildcard patterns, specified relative to the project's root, may be processed for gathering. */
 	UPROPERTY(config, EditAnywhere, Category = "Filter")
 	TArray<FString> IncludePathWildcards;
 
+	/* Packages whose paths match these wildcard patterns will be excluded from gathering. */
 	UPROPERTY(config, EditAnywhere, Category = "Filter")
 	TArray<FString> ExcludePathWildcards;
 
+	/* Packages whose names match these wildcard patterns may be processed for text to gather. */
 	UPROPERTY(config, EditAnywhere, Category = "Filter")
 	TArray<FString> FileExtensions;
 };
@@ -48,12 +54,18 @@ struct FMetaDataKeyGatherSpecification
 {
 	GENERATED_USTRUCT_BODY()
 
+	/*  The metadata key for which values will be gathered as text. */
 	UPROPERTY(config, EditAnywhere, Category = "Input")
 	FString MetaDataKey;
 
+	/* The localization namespace in which the gathered text will be output. */
 	UPROPERTY(config, EditAnywhere, Category = "Output")
 	FString TextNamespace;
 
+	/* The pattern which will be formatted to form the localization key for the metadata value gathered as text.
+	Placeholder - Description
+	{FieldPath} - The fully qualified name of the object upon which the metadata resides.
+	{MetaDataValue} - The value associated with the metadata key. */
 	UPROPERTY(config, EditAnywhere, Category = "Output")
 	FString TextKeyPattern;
 };
@@ -63,12 +75,15 @@ struct FGatherTextFromMetaDataConfiguration
 {
 	GENERATED_USTRUCT_BODY()
 
+	/* Metadata from source files whose paths match these wildcard patterns, specified relative to the project's root, may be processed for gathering. */
 	UPROPERTY(config, EditAnywhere, Category = "Filter")
 	TArray<FString> IncludePathWildcards;
 
+	/* Metadata from source files whose paths match these wildcard patterns will be excluded from gathering. */
 	UPROPERTY(config, EditAnywhere, Category = "Filter")
 	TArray<FString> ExcludePathWildcards;
 
+	/* Specifications for how to gather text from specific metadata keys. */
 	UPROPERTY(config, EditAnywhere, Category = "MetaData")
 	TArray<FMetaDataKeyGatherSpecification> KeySpecifications;
 };
@@ -83,9 +98,11 @@ struct FCultureStatistics
 	{
 	}
 
+	/* The ISO name for this culture. */
 	UPROPERTY(config, EditAnywhere, Category = "Culture")
 	FString CultureName;
 
+	/* The estimated number of words that have been localized for this culture. */
 	UPROPERTY(Transient, EditAnywhere, Category = "Statistics")
 	uint32 WordCount;
 };
@@ -93,12 +110,12 @@ struct FCultureStatistics
 UENUM()
 enum class ELocalizationTargetLoadingPolicy : uint8
 {
-	Never,
-	Always,
-	Editor,
-	Game,
-	PropertyNames,
-	ToolTips,
+	Never,			/* This target's localization data will never be loaded automatically. */
+	Always,			/* This target's localization data will always be loaded automatically. */
+	Editor,			/* This target's localization data will only be loaded when running the editor. Use if this target localizes the editor. */
+	Game,			/* This target's localization data will only be loaded when running the game. Use if this target localizes your game. */
+	PropertyNames,	/* This target's localization data will only be loaded if the editor is displaying localized property names. */
+	ToolTips,		/* This target's localization data will only be loaded if the editor is displaying localized tool tips. */
 };
 
 USTRUCT()
@@ -108,37 +125,46 @@ struct FLocalizationTargetSettings
 
 	FLocalizationTargetSettings();
 
-	/* Unique name for the target within the context of the project. */
+	/* Unique name for the target. */
 	UPROPERTY(config, EditAnywhere, Category = "Target")
 	FString Name;
 
 	UPROPERTY(config)
 	FGuid Guid;
 
+	/* Whether the target has outstanding conflicts that require resolution. */
 	UPROPERTY(Transient, EditAnywhere, Category = "Target")
 	ELocalizationTargetConflictStatus ConflictStatus;
 
-	UPROPERTY(config, EditAnywhere, Category = "Gather Configuration")
+	/* Text present in these targets will not be duplicated in this target. */
+	UPROPERTY(config, EditAnywhere, Category = "Gather")
 	TArray<FGuid> TargetDependencies;
 
-	UPROPERTY(config, EditAnywhere, Category = "Gather Configuration", AdvancedDisplay, meta=(FilePathFilter="manifest"))
+	/* Text present in these manifests will not be duplicated in this target. */
+	UPROPERTY(config, EditAnywhere, Category = "Gather", AdvancedDisplay, meta=(FilePathFilter="manifest"))
 	TArray<FFilePath> AdditionalManifestDependencies;
 
-	UPROPERTY(config, EditAnywhere, Category = "Gather Configuration", AdvancedDisplay)
+	/* The names of modules which must be loaded when gathering text. Use to gather from packages or metadata sourced from a module that isn't the primary game module. */
+	UPROPERTY(config, EditAnywhere, Category = "Gather", AdvancedDisplay)
 	TArray<FString> RequiredModuleNames;
 
-	UPROPERTY(config, EditAnywhere, Category = "Gather Configuration")
+	/* Parameters for defining what text is gathered from text files and how. */
+	UPROPERTY(config, EditAnywhere, Category = "Gather")
 	FGatherTextFromTextFilesConfiguration GatherFromTextFiles;
 
-	UPROPERTY(config, EditAnywhere, Category = "Gather Configuration")
+	/* Parameters for defining what text is gathered from packages and how. */
+	UPROPERTY(config, EditAnywhere, Category = "Gather")
 	FGatherTextFromPackagesConfiguration GatherFromPackages;
 
-	UPROPERTY(config, EditAnywhere, Category = "Gather Configuration")
+	/* Parameters for defining what text is gathered from metadata and how. */
+	UPROPERTY(config, EditAnywhere, Category = "Gather")
 	FGatherTextFromMetaDataConfiguration GatherFromMetaData;
 
+	/* The culture in which the source text is written in. */
 	UPROPERTY(config, EditAnywhere, Category = "Cultures")
 	FCultureStatistics NativeCultureStatistics;
 
+	/* Cultures for which the source text is being localized for.*/
 	UPROPERTY(config, EditAnywhere, Category = "Cultures")
 	TArray<FCultureStatistics> SupportedCulturesStatistics;
 };
