Code improvements for localization targets in preparation for storing all settings in a single object for the localization dashboard.

diff --git a/Engine/Source/Editor/LocalizationDashboard/Classes/LocalizationTarget.h b/Engine/Source/Editor/LocalizationDashboard/Classes/LocalizationTarget.h
--- a/Engine/Source/Editor/LocalizationDashboard/Classes/LocalizationTarget.h
+++ b/Engine/Source/Editor/LocalizationDashboard/Classes/LocalizationTarget.h
@@ -174,15 +174,23 @@ class ULocalizationTarget : public UObject
 {
 	GENERATED_BODY()
 
+public:
+	ULocalizationTarget(const FObjectInitializer& ObjectInitializer);
+
 public:
 	UPROPERTY(EditAnywhere, Category = "Target")
 	FLocalizationTargetSettings Settings;
 
+private:
+	FLocalizationTargetSettings* SettingsSource;
+
 public:
 #if WITH_EDITOR
 	virtual void PostEditChangeProperty(FPropertyChangedEvent& PropertyChangedEvent) override;
 #endif
 
+	void SetSettingsSource(FLocalizationTargetSettings& SettingsSource);
+
 	bool IsMemberOfEngineTargetSet() const;
 	bool UpdateWordCountsFromCSV();
 	void UpdateStatusFromConflictReport();
@@ -204,7 +212,7 @@ public:
 
 private:
 	UPROPERTY(config)
-	TArray<FLocalizationTargetSettings> Targets;
+	TArray<FLocalizationTargetSettings> TargetsSettings;
 
 public:
 #if WITH_EDITOR
diff --git a/Engine/Source/Editor/LocalizationDashboard/Private/LocalizationTarget.cpp b/Engine/Source/Editor/LocalizationDashboard/Private/LocalizationTarget.cpp
--- a/Engine/Source/Editor/LocalizationDashboard/Private/LocalizationTarget.cpp
+++ b/Engine/Source/Editor/LocalizationDashboard/Private/LocalizationTarget.cpp
@@ -1,4 +1,4 @@
-// Copyright 1998-2015 Epic Games, Inc. All Rights Reserved.
+﻿// Copyright 1998-2015 Epic Games, Inc. All Rights Reserved.
 
 #include "LocalizationDashboardPrivatePCH.h"
 #include "Classes/LocalizationTarget.h"
@@ -156,8 +156,27 @@ namespace
 	};
 }
 
+ULocalizationTarget::ULocalizationTarget(const FObjectInitializer& ObjectInitializer)
+	: Super(ObjectInitializer)
+	, SettingsSource(nullptr)
+{
+
+}
+
 void ULocalizationTarget::PostEditChangeProperty(FPropertyChangedEvent& PropertyChangedEvent)
 {
+	Super::PostEditChangeProperty(PropertyChangedEvent);
+
+	// Copy change back to source.
+	if (PropertyChangedEvent.MemberProperty->GetFName() == GET_MEMBER_NAME_CHECKED(ULocalizationTarget,Settings))
+	{
+		if (SettingsSource)
+		{
+			*SettingsSource = Settings;
+		}
+	}
+
+	// Notify parent of change.
 	ULocalizationTargetSet* const LocalizationTargetSet = Cast<ULocalizationTargetSet>(GetOuter());
 	if (LocalizationTargetSet)
 	{
@@ -169,6 +188,12 @@ void ULocalizationTarget::PostEditChangeProperty(FPropertyChangedEvent& Property
 	}
 }
 
+void ULocalizationTarget::SetSettingsSource(FLocalizationTargetSettings& InSettingsSource)
+{
+	SettingsSource = &InSettingsSource;
+	Settings = *SettingsSource;
+}
+
 bool ULocalizationTarget::IsMemberOfEngineTargetSet() const
 {
 	ULocalizationTargetSet* const TargetSet = CastChecked<ULocalizationTargetSet>(GetOuter());
@@ -458,11 +483,12 @@ void ULocalizationTargetSet::PostInitProperties()
 {
 	Super::PostInitProperties();
 
-	TargetObjects.Empty(Targets.Num());
-	for (const FLocalizationTargetSettings& Target : Targets)
+	// Create UObjects used for details views and associate them with their backing config settings.
+	TargetObjects.Empty(TargetsSettings.Num());
+	for (FLocalizationTargetSettings& TargetSettings : TargetsSettings)
 	{
 		ULocalizationTarget* const NewTargetObject = NewObject<ULocalizationTarget>(this);
-		NewTargetObject->Settings = Target;
+		NewTargetObject->SetSettingsSource(TargetSettings);
 		TargetObjects.Add(NewTargetObject);
 	}
 }
@@ -471,14 +497,14 @@ void ULocalizationTargetSet::PostEditChangeProperty(FPropertyChangedEvent& Prope
 {
 	Super::PostEditChangeProperty(PropertyChangedEvent);
 
-	if (PropertyChangedEvent.MemberProperty->GetFName() == GET_MEMBER_NAME_CHECKED(ULocalizationTargetSet,TargetObjects))
-	{
-		Targets.Empty(TargetObjects.Num());
-		for (ULocalizationTarget* const& TargetObject : TargetObjects)
-		{
-			Targets.Add(TargetObject ? TargetObject->Settings : FLocalizationTargetSettings());
-		}
-	}
+	//if (PropertyChangedEvent.MemberProperty->GetFName() == GET_MEMBER_NAME_CHECKED(ULocalizationTargetSet,TargetObjects))
+	//{
+	//	TargetsSettings.Empty(TargetObjects.Num());
+	//	for (ULocalizationTarget* const& TargetObject : TargetObjects)
+	//	{
+	//		TargetsSettings.Add(TargetObject ? TargetObject->Settings : FLocalizationTargetSettings());
+	//	}
+	//}
 
 	ISourceControlModule& SourceControl = ISourceControlModule::Get();
 	ISourceControlProvider& SourceControlProvider = SourceControl.GetProvider();
