Refactored localization target data structures so that the configuration settings exist within a unified localization dashboard settings object.

diff --git a/Engine/Config/BaseEditor.ini b/Engine/Config/BaseEditor.ini
--- a/Engine/Config/BaseEditor.ini
+++ b/Engine/Config/BaseEditor.ini
@@ -356,5 +356,12 @@ MaxMemoryAllowance=8192
 +VersionedIniParams=Engine:/Script/Engine.Engine:WeightMapPlaceholderTextureName
 +VersionedIniParams=Engine:/Script/Engine.Engine:LightMapDensityTextureName
 
+[/Script/LocalizationDashboard.LocalizationDashboardSettings]
++EngineTargetsSettings=(Name="Engine",Guid=33482D004789784C9DA695A682ACCA1B,TargetDependencies=,AdditionalManifestDependencies=,RequiredModuleNames=,GatherFromTextFiles=(SearchDirectories=("Source/Runtime/","Source/Developer/"),ExcludePathWildcards=,FileExtensions=("*.cpp","*.h","*.c","*.inl","*.mm","*.ini")),GatherFromPackages=(IncludePathWildcards=("Content/"),ExcludePathWildcards=("*/Content/Editor/*","*/Content/Tutorial/*","*/Content/Developers/*","*/Content/TestPackages/*","*/Content/QA_Assets/*","*/Content/Maps/Automation/*"),FileExtensions=("*.umap","*.uasset")),GatherFromMetaData=(IncludePathWildcards=,ExcludePathWildcards=,KeySpecifications=),NativeCultureStatistics=(CultureName="en"),SupportedCulturesStatistics=((CultureName="es"),(CultureName="ja"),(CultureName="ko"),(CultureName="pt-BR"),(CultureName="zh-CN")))
++EngineTargetsSettings=(Name="Editor",Guid=AC8BFD2A41A2FB2893BB8EA0AF903E6D,TargetDependencies=(33482D004789784C9DA695A682ACCA1B),AdditionalManifestDependencies=,RequiredModuleNames=,GatherFromTextFiles=(SearchDirectories=("Source/Editor/"),ExcludePathWildcards=,FileExtensions=("*.cpp","*.h","*.c","*.inl","*.mm","*.ini")),GatherFromPackages=(IncludePathWildcards=("Content/Editor/*","Content/Editor*"),ExcludePathWildcards=,FileExtensions=("*.umap","*.uasset")),GatherFromMetaData=(IncludePathWildcards=("Source/Editor/*","Source/Runtime/*","Source/Developer/*"),ExcludePathWildcards=,KeySpecifications=((MetaDataKey="Category",TextNamespace="UObjectCategories",TextKeyPattern="{MetaDataValue}"))),NativeCultureStatistics=(CultureName="en"),SupportedCulturesStatistics=((CultureName="es"),(CultureName="ja"),(CultureName="ko"),(CultureName="pt-BR"),(CultureName="zh-CN")))
++EngineTargetsSettings=(Name="EditorTutorials",Guid=00F8E3AD47F0A73D50D46881C14DF28F,TargetDependencies=(33482D004789784C9DA695A682ACCA1B,AC8BFD2A41A2FB2893BB8EA0AF903E6D),AdditionalManifestDependencies=,RequiredModuleNames=,GatherFromTextFiles=(SearchDirectories=,ExcludePathWildcards=,FileExtensions=),GatherFromPackages=(IncludePathWildcards=("Content/Tutorial/*"),ExcludePathWildcards=,FileExtensions=("*.umap","*.uasset")),GatherFromMetaData=(IncludePathWildcards=,ExcludePathWildcards=,KeySpecifications=),NativeCultureStatistics=(CultureName="en"),SupportedCulturesStatistics=((CultureName="es"),(CultureName="ja"),(CultureName="ko"),(CultureName="pt-BR"),(CultureName="zh-CN")))
++EngineTargetsSettings=(Name="PropertyNames",Guid=E391A8B149980E8154E056AF2DA49479,TargetDependencies=(33482D004789784C9DA695A682ACCA1B),AdditionalManifestDependencies=,RequiredModuleNames=,GatherFromTextFiles=(SearchDirectories=,ExcludePathWildcards=,FileExtensions=),GatherFromPackages=(IncludePathWildcards=,ExcludePathWildcards=,FileExtensions=),GatherFromMetaData=(IncludePathWildcards=("Source/Editor/*","Source/Runtime/*","Source/Developer/*"),ExcludePathWildcards=("*/Source/Developer/NoRedist/*"),KeySpecifications=((MetaDataKey="DisplayName",TextNamespace="UObjectDisplayNames",TextKeyPattern="{FieldPath}"))),NativeCultureStatistics=(CultureName="en"),SupportedCulturesStatistics=((CultureName="es"),(CultureName="ja"),(CultureName="ko"),(CultureName="pt-BR"),(CultureName="zh-CN")))
++EngineTargetsSettings=(Name="ToolTips",Guid=0F116534468918AEA432DD8C77703BA8,TargetDependencies=(33482D004789784C9DA695A682ACCA1B,AC8BFD2A41A2FB2893BB8EA0AF903E6D),AdditionalManifestDependencies=,RequiredModuleNames=,GatherFromTextFiles=(SearchDirectories=,ExcludePathWildcards=,FileExtensions=),GatherFromPackages=(IncludePathWildcards=,ExcludePathWildcards=,FileExtensions=),GatherFromMetaData=(IncludePathWildcards=("Source/Editor/*","Source/Runtime/*","Source/Developer/*"),ExcludePathWildcards=("*/Source/Developer/NoRedist/*"),KeySpecifications=((MetaDataKey="ToolTip",TextNamespace="UObjectToolTips",TextKeyPattern="{FieldPath}"),(MetaDataKey="ShortToolTip",TextNamespace="UObjectShortToolTips",TextKeyPattern="{FieldPath}"))),NativeCultureStatistics=(CultureName="en"),SupportedCulturesStatistics=((CultureName="es"),(CultureName="ja"),(CultureName="ko"),(CultureName="pt-BR"),(CultureName="zh-CN")))
++GameTargetsSettings=(Name="Game",Guid=AE0EA34A45461A25BA65A391026F19F8,TargetDependencies=(33482D004789784C9DA695A682ACCA1B,AC8BFD2A41A2FB2893BB8EA0AF903E6D),AdditionalManifestDependencies=,RequiredModuleNames=,GatherFromTextFiles=(SearchDirectories=,ExcludePathWildcards=,FileExtensions=),GatherFromPackages=(IncludePathWildcards=,ExcludePathWildcards=,FileExtensions=),GatherFromMetaData=(IncludePathWildcards=,ExcludePathWildcards=,KeySpecifications=),NativeCultureStatistics=(CultureName="en"),SupportedCulturesStatistics=)
 
 
diff --git a/Engine/Source/Editor/LocalizationDashboard/Classes/LocalizationTarget.h b/Engine/Source/Editor/LocalizationDashboard/Classes/LocalizationTarget.h
deleted file mode 100644
diff --git a/Engine/Source/Editor/LocalizationDashboard/Private/LocalizationCommandletExecution.cpp b/Engine/Source/Editor/LocalizationDashboard/Private/LocalizationCommandletExecution.cpp
--- a/Engine/Source/Editor/LocalizationDashboard/Private/LocalizationCommandletExecution.cpp
+++ b/Engine/Source/Editor/LocalizationDashboard/Private/LocalizationCommandletExecution.cpp
@@ -2,7 +2,7 @@
 
 #include "LocalizationDashboardPrivatePCH.h"
 #include "LocalizationCommandletExecution.h"
-#include "Classes/LocalizationTarget.h"
+#include "LocalizationTargetTypes.h"
 #include "LocalizationConfigurationScript.h"
 #include "ISourceControlModule.h"
 #include "DesktopPlatformModule.h"
diff --git a/Engine/Source/Editor/LocalizationDashboard/Private/LocalizationCommandletTasks.h b/Engine/Source/Editor/LocalizationDashboard/Private/LocalizationCommandletTasks.h
--- a/Engine/Source/Editor/LocalizationDashboard/Private/LocalizationCommandletTasks.h
+++ b/Engine/Source/Editor/LocalizationDashboard/Private/LocalizationCommandletTasks.h
@@ -1,7 +1,7 @@
 // Copyright 1998-2015 Epic Games, Inc. All Rights Reserved.
 #pragma once
 
-#include "Classes/LocalizationTarget.h"
+#include "LocalizationTargetTypes.h"
 
 namespace LocalizationCommandletTasks
 {
diff --git a/Engine/Source/Editor/LocalizationDashboard/Private/LocalizationConfigurationScript.cpp b/Engine/Source/Editor/LocalizationDashboard/Private/LocalizationConfigurationScript.cpp
--- a/Engine/Source/Editor/LocalizationDashboard/Private/LocalizationConfigurationScript.cpp
+++ b/Engine/Source/Editor/LocalizationDashboard/Private/LocalizationConfigurationScript.cpp
@@ -2,7 +2,8 @@
 
 #include "LocalizationDashboardPrivatePCH.h"
 #include "LocalizationConfigurationScript.h"
-#include "Classes/LocalizationTarget.h"
+#include "LocalizationTargetTypes.h"
+#include "LocalizationDashboardSettings.h"
 
 namespace
 {
@@ -97,7 +98,7 @@ namespace LocalizationConfigurationScript
 			for (const FGuid& TargetDependencyGuid : Target->Settings.TargetDependencies)
 			{
 				TArray<ULocalizationTarget*> AllLocalizationTargets;
-				ULocalizationTargetSet* EngineTargetSet = FindObjectChecked<ULocalizationTargetSet>(ANY_PACKAGE, *ULocalizationTargetSet::EngineTargetSetName.ToString());
+				ULocalizationTargetSet* EngineTargetSet = ULocalizationDashboardSettings::GetEngineTargetSet();
 				if (EngineTargetSet != LocalizationTargetSet)
 				{
 					AllLocalizationTargets.Append(EngineTargetSet->TargetObjects);
diff --git a/Engine/Source/Editor/LocalizationDashboard/Private/LocalizationConfigurationScript.h b/Engine/Source/Editor/LocalizationDashboard/Private/LocalizationConfigurationScript.h
--- a/Engine/Source/Editor/LocalizationDashboard/Private/LocalizationConfigurationScript.h
+++ b/Engine/Source/Editor/LocalizationDashboard/Private/LocalizationConfigurationScript.h
@@ -3,7 +3,7 @@
 #pragma once
 
 #include "ConfigCacheIni.h"
-#include "Classes/LocalizationTarget.h"
+#include "LocalizationTargetTypes.h"
 
 struct FLocalizationConfigurationScript : public FConfigFile
 {
diff --git a/Engine/Source/Editor/LocalizationDashboard/Private/LocalizationDashboard.cpp b/Engine/Source/Editor/LocalizationDashboard/Private/LocalizationDashboard.cpp
--- a/Engine/Source/Editor/LocalizationDashboard/Private/LocalizationDashboard.cpp
+++ b/Engine/Source/Editor/LocalizationDashboard/Private/LocalizationDashboard.cpp
@@ -8,7 +8,8 @@
 #include "IMainFrameModule.h"
 #include "ModuleManager.h"
 #include "SLocalizationTargetEditor.h"
-#include "LocalizationTarget.h"
+#include "LocalizationDashboardSettings.h"
+#include "LocalizationTargetTypes.h"
 
 #define LOCTEXT_NAMESPACE "LocalizationDashboard"
 
@@ -54,12 +55,12 @@ void SLocalizationDashboard::Construct(const FArguments& InArgs, const TSharedPt
 		if (SpawnTabArgs.GetTabId() == EngineTargetSetDetailsTabName)
 		{
 			DockTabLabel = LOCTEXT("EngineTargetsTabLabel", "Engine Targets");
-			TargetSet = FindObjectChecked<ULocalizationTargetSet>(GetTransientPackage(), TEXT("EngineLocalizationTargetSet"));
+			TargetSet = ULocalizationDashboardSettings::GetEngineTargetSet();
 		}
 		else if (SpawnTabArgs.GetTabId() == ProjectTargetSetDetailsTabName)
 		{
 			DockTabLabel = LOCTEXT("ProjectTargetsTabLabel", "Project Targets");
-			TargetSet = FindObjectChecked<ULocalizationTargetSet>(GetTransientPackage(), TEXT("ProjectLocalizationTargetSet"));
+			TargetSet = ULocalizationDashboardSettings::GetGameTargetSet();
 		}
 
 		const TSharedRef<SDockTab> DockTab = SNew(SDockTab)
diff --git a/Engine/Source/Editor/LocalizationDashboard/Private/LocalizationDashboardModule.cpp b/Engine/Source/Editor/LocalizationDashboard/Private/LocalizationDashboardModule.cpp
--- a/Engine/Source/Editor/LocalizationDashboard/Private/LocalizationDashboardModule.cpp
+++ b/Engine/Source/Editor/LocalizationDashboard/Private/LocalizationDashboardModule.cpp
@@ -3,7 +3,7 @@
 #include "LocalizationDashboardPrivatePCH.h"
 #include "ILocalizationDashboardModule.h"
 #include "LocalizationDashboard.h"
-#include "LocalizationTarget.h"
+#include "LocalizationTargetTypes.h"
 #include "SDockTab.h"
 #include "Features/IModularFeatures.h"
 #include "ILocalizationServiceProvider.h"
@@ -19,21 +19,6 @@ public:
 	// Begin IModuleInterface interface
 	virtual void StartupModule() override
 	{
-		EngineTargetSet = NewObject<ULocalizationTargetSet>(GetTransientPackage(), ULocalizationTargetSet::EngineTargetSetName, RF_RootSet);
-		ProjectTargetSet = NewObject<ULocalizationTargetSet>(GetTransientPackage(), ULocalizationTargetSet::ProjectTargetSetName, RF_RootSet);
-		
-		for (ULocalizationTarget*& Target : EngineTargetSet->TargetObjects)
-		{
-			Target->UpdateStatusFromConflictReport();
-			Target->UpdateWordCountsFromCSV();
-		}
-
-		for (ULocalizationTarget*& Target : ProjectTargetSet->TargetObjects)
-		{
-			Target->UpdateStatusFromConflictReport();
-			Target->UpdateWordCountsFromCSV();
-	}
-
 		ServiceProviders = IModularFeatures::Get().GetModularFeatureImplementations<ILocalizationServiceProvider>("LocalizationService");
 		ServiceProviders.Insert(nullptr, 0); // "None"
 
@@ -94,9 +79,6 @@ public:
 	}
 
 private:
-	TWeakObjectPtr<ULocalizationTargetSet> EngineTargetSet;
-	TWeakObjectPtr<ULocalizationTargetSet> ProjectTargetSet;
-
 	TArray<ILocalizationServiceProvider*> ServiceProviders;
 	FName CurrentServiceProviderName;
 };
diff --git a/Engine/Source/Editor/LocalizationDashboard/Private/LocalizationTarget.cpp b/Engine/Source/Editor/LocalizationDashboard/Private/LocalizationTarget.cpp
--- a/Engine/Source/Editor/LocalizationDashboard/Private/LocalizationTarget.cpp
+++ b/Engine/Source/Editor/LocalizationDashboard/Private/LocalizationTarget.cpp
@@ -1,11 +1,10 @@
 ﻿// Copyright 1998-2015 Epic Games, Inc. All Rights Reserved.
 
 #include "LocalizationDashboardPrivatePCH.h"
-#include "Classes/LocalizationTarget.h"
-#include "Classes/LocalizationTarget.h"
-#include "ISourceControlModule.h"
+#include "LocalizationTargetTypes.h"
 #include "CsvParser.h"
 #include "LocalizationConfigurationScript.h"
+#include "LocalizationDashboardSettings.h"
 
 FLocalizationTargetSettings::FLocalizationTargetSettings()
 	: Guid(FGuid::NewGuid())
@@ -156,48 +155,22 @@ namespace
 	};
 }
 
-ULocalizationTarget::ULocalizationTarget(const FObjectInitializer& ObjectInitializer)
-	: Super(ObjectInitializer)
-	, SettingsSource(nullptr)
-{
-
-}
-
 void ULocalizationTarget::PostEditChangeProperty(FPropertyChangedEvent& PropertyChangedEvent)
 {
 	Super::PostEditChangeProperty(PropertyChangedEvent);
 
-	// Copy change back to source.
-	if (PropertyChangedEvent.MemberProperty->GetFName() == GET_MEMBER_NAME_CHECKED(ULocalizationTarget,Settings))
-	{
-		if (SettingsSource)
-		{
-			*SettingsSource = Settings;
-		}
-	}
-
 	// Notify parent of change.
 	ULocalizationTargetSet* const LocalizationTargetSet = Cast<ULocalizationTargetSet>(GetOuter());
 	if (LocalizationTargetSet)
 	{
-		UClass* ParentClass = LocalizationTargetSet ? LocalizationTargetSet->GetClass() : nullptr;
-		UProperty* ChangedParentProperty = ParentClass ? ParentClass->FindPropertyByName(GET_MEMBER_NAME_CHECKED(ULocalizationTargetSet,TargetObjects)) : nullptr;
-
-		FPropertyChangedEvent PropertyChangedEventForParent(ChangedParentProperty);
-		LocalizationTargetSet->PostEditChangeProperty(PropertyChangedEventForParent);
+		LocalizationTargetSet->PostEditChange();
 	}
 }
 
-void ULocalizationTarget::SetSettingsSource(FLocalizationTargetSettings& InSettingsSource)
-{
-	SettingsSource = &InSettingsSource;
-	Settings = *SettingsSource;
-}
-
 bool ULocalizationTarget::IsMemberOfEngineTargetSet() const
 {
 	ULocalizationTargetSet* const TargetSet = CastChecked<ULocalizationTargetSet>(GetOuter());
-	return TargetSet == FindObjectChecked<ULocalizationTargetSet>(GetTransientPackage(), *ULocalizationTargetSet::EngineTargetSetName.ToString());
+	return TargetSet == ULocalizationDashboardSettings::GetEngineTargetSet();
 }
 
 bool ULocalizationTarget::UpdateWordCountsFromCSV()
@@ -476,54 +449,14 @@ bool ULocalizationTarget::DeleteFiles(const FString* const Culture) const
 	return HasCompletelySucceeded;
 }
 
-FName ULocalizationTargetSet::EngineTargetSetName = TEXT("EngineLocalizationTargetSet");
-FName ULocalizationTargetSet::ProjectTargetSetName = TEXT("ProjectLocalizationTargetSet");
-
-void ULocalizationTargetSet::PostInitProperties()
-{
-	Super::PostInitProperties();
-
-	// Create UObjects used for details views and associate them with their backing config settings.
-	TargetObjects.Empty(TargetsSettings.Num());
-	for (FLocalizationTargetSettings& TargetSettings : TargetsSettings)
-	{
-		ULocalizationTarget* const NewTargetObject = NewObject<ULocalizationTarget>(this);
-		NewTargetObject->SetSettingsSource(TargetSettings);
-		TargetObjects.Add(NewTargetObject);
-	}
-}
-
 void ULocalizationTargetSet::PostEditChangeProperty(FPropertyChangedEvent& PropertyChangedEvent)
 {
 	Super::PostEditChangeProperty(PropertyChangedEvent);
 
-	//if (PropertyChangedEvent.MemberProperty->GetFName() == GET_MEMBER_NAME_CHECKED(ULocalizationTargetSet,TargetObjects))
-	//{
-	//	TargetsSettings.Empty(TargetObjects.Num());
-	//	for (ULocalizationTarget* const& TargetObject : TargetObjects)
-	//	{
-	//		TargetsSettings.Add(TargetObject ? TargetObject->Settings : FLocalizationTargetSettings());
-	//	}
-	//}
-
-	ISourceControlModule& SourceControl = ISourceControlModule::Get();
-	ISourceControlProvider& SourceControlProvider = SourceControl.GetProvider();
-	const bool CanUseSourceControl = SourceControl.IsEnabled() && SourceControlProvider.IsEnabled() && SourceControlProvider.IsAvailable();
-
-	if (CanUseSourceControl)
+	// Notify parent of change.
+	ULocalizationDashboardSettings* const LocalizationDashboardSettings = Cast<ULocalizationDashboardSettings>(GetOuter());
+	if (LocalizationDashboardSettings)
 	{
-		const FString Path = GetConfigFilename(this);
-		const FSourceControlStatePtr SourceControlState = SourceControlProvider.GetState(*Path, EStateCacheUsage::Use);
-
-		// Should only use source control if the file is actually under source control.
-		if (CanUseSourceControl && SourceControlState.IsValid() && !SourceControlState->CanAdd())
-		{
-			if (!SourceControlState->CanEdit())
-			{
-				SourceControlProvider.Execute(ISourceControlOperation::Create<FCheckOut>(), *Path);
-			}
-		}
+		LocalizationDashboardSettings->PostEditChange();
 	}
-
-	SaveConfig();
 }
\ No newline at end of file
diff --git a/Engine/Source/Editor/LocalizationDashboard/Private/LocalizationTargetDetailCustomization.cpp b/Engine/Source/Editor/LocalizationDashboard/Private/LocalizationTargetDetailCustomization.cpp
--- a/Engine/Source/Editor/LocalizationDashboard/Private/LocalizationTargetDetailCustomization.cpp
+++ b/Engine/Source/Editor/LocalizationDashboard/Private/LocalizationTargetDetailCustomization.cpp
@@ -3,7 +3,7 @@
 #include "LocalizationDashboardPrivatePCH.h"
 #include "LocalizationTargetDetailCustomization.h"
 #include "DetailLayoutBuilder.h"
-#include "LocalizationTarget.h"
+#include "LocalizationTargetTypes.h"
 #include "DetailWidgetRow.h"
 #include "SLocalizationTargetStatusButton.h"
 #include "SLocalizationTargetEditorCultureRow.h"
@@ -13,6 +13,7 @@
 #include "IDetailsView.h"
 #include "LocalizationCommandletTasks.h"
 #include "ObjectEditorUtils.h"
+#include "LocalizationDashboardSettings.h"
 
 #define LOCTEXT_NAMESPACE "LocalizationTargetEditor"
 
@@ -404,7 +405,7 @@ FText FLocalizationTargetDetailCustomization::GetTargetName() const
 bool FLocalizationTargetDetailCustomization::IsTargetNameUnique(const FString& Name) const
 {
 	TArray<ULocalizationTarget*> AllLocalizationTargets;
-	ULocalizationTargetSet* EngineTargetSet = FindObjectChecked<ULocalizationTargetSet>(ANY_PACKAGE, *ULocalizationTargetSet::EngineTargetSetName.ToString());
+	ULocalizationTargetSet* EngineTargetSet = ULocalizationDashboardSettings::GetEngineTargetSet();
 	if (EngineTargetSet != TargetSet)
 	{
 		AllLocalizationTargets.Append(EngineTargetSet->TargetObjects);
@@ -548,7 +549,7 @@ void FLocalizationTargetDetailCustomization::RebuildTargetDependenciesBox()
 		TargetDependenciesWidgets.Empty();
 
 		TArray<ULocalizationTarget*> AllLocalizationTargets;
-		ULocalizationTargetSet* EngineTargetSet = FindObjectChecked<ULocalizationTargetSet>(ANY_PACKAGE, *ULocalizationTargetSet::EngineTargetSetName.ToString());
+		ULocalizationTargetSet* EngineTargetSet = ULocalizationDashboardSettings::GetEngineTargetSet();
 		if (EngineTargetSet != TargetSet)
 		{
 			AllLocalizationTargets.Append(EngineTargetSet->TargetObjects);
@@ -602,7 +603,7 @@ void FLocalizationTargetDetailCustomization::RebuildTargetsList()
 	};
 
 	TArray<ULocalizationTarget*> AllLocalizationTargets;
-	ULocalizationTargetSet* EngineTargetSet = FindObjectChecked<ULocalizationTargetSet>(ANY_PACKAGE, *ULocalizationTargetSet::EngineTargetSetName.ToString());
+	ULocalizationTargetSet* EngineTargetSet = ULocalizationDashboardSettings::GetEngineTargetSet();
 	if (EngineTargetSet != TargetSet)
 	{
 		AllLocalizationTargets.Append(EngineTargetSet->TargetObjects);
diff --git a/Engine/Source/Editor/LocalizationDashboard/Private/LocalizationTargetDetailCustomization.h b/Engine/Source/Editor/LocalizationDashboard/Private/LocalizationTargetDetailCustomization.h
--- a/Engine/Source/Editor/LocalizationDashboard/Private/LocalizationTargetDetailCustomization.h
+++ b/Engine/Source/Editor/LocalizationDashboard/Private/LocalizationTargetDetailCustomization.h
@@ -2,7 +2,7 @@
 
 #pragma once
 #include "IDetailCustomization.h"
-#include "LocalizationTarget.h"
+#include "LocalizationTargetTypes.h"
 
 class SCulturePicker;
 class IPropertyHandle;
diff --git a/Engine/Source/Editor/LocalizationDashboard/Private/LocalizationTargetSetDetailCustomization.cpp b/Engine/Source/Editor/LocalizationDashboard/Private/LocalizationTargetSetDetailCustomization.cpp
--- a/Engine/Source/Editor/LocalizationDashboard/Private/LocalizationTargetSetDetailCustomization.cpp
+++ b/Engine/Source/Editor/LocalizationDashboard/Private/LocalizationTargetSetDetailCustomization.cpp
@@ -3,7 +3,7 @@
 #include "LocalizationDashboardPrivatePCH.h"
 #include "LocalizationTargetSetDetailCustomization.h"
 #include "DetailLayoutBuilder.h"
-#include "LocalizationTarget.h"
+#include "LocalizationTargetTypes.h"
 #include "DetailWidgetRow.h"
 #include "SLocalizationDashboardTargetRow.h"
 #include "DetailCategoryBuilder.h"
diff --git a/Engine/Source/Editor/LocalizationDashboard/Private/LocalizationTargetTypes.h b/Engine/Source/Editor/LocalizationDashboard/Private/LocalizationTargetTypes.h
new file mode 100644
--- /dev/null
+++ b/Engine/Source/Editor/LocalizationDashboard/Private/LocalizationTargetTypes.h
@@ -0,0 +1,206 @@
+﻿// Copyright 1998-2015 Epic Games, Inc. All Rights Reserved.
+
+#pragma once
+
+#include "Engine/EngineTypes.h"
+#include "LocalizationTargetTypes.generated.h"
+
+UENUM()
+enum class ELocalizationTargetConflictStatus : uint8
+{
+	Unknown, /* The status of conflicts in this localization target could not be determined. */
+	ConflictsPresent, /* The are outstanding conflicts present in this localization target. */
+	Clear, /* The localization target is clear of conflicts. */
+};
+
+USTRUCT()
+struct FGatherTextFromTextFilesConfiguration
+{
+	GENERATED_USTRUCT_BODY()
+
+	/* The paths of directories to be searched for text files, specified relative to the project's root, which may be parsed for text to gather. */
+	UPROPERTY(config, EditAnywhere, Category = "Filter")
+	TArray<FString> SearchDirectories;
+
+	/* Text files whose paths match these wildcard patterns will be excluded from gathering. */
+	UPROPERTY(config, EditAnywhere, Category = "Filter")
+	TArray<FString> ExcludePathWildcards;
+
+	/* Text files whose names match these wildcard patterns may be parsed for text to gather. */
+	UPROPERTY(config, EditAnywhere, Category = "Filter")
+	TArray<FString> FileExtensions;
+};
+
+USTRUCT()
+struct FGatherTextFromPackagesConfiguration
+{
+	GENERATED_USTRUCT_BODY()
+
+	/* Packages whose paths match these wildcard patterns, specified relative to the project's root, may be processed for gathering. */
+	UPROPERTY(config, EditAnywhere, Category = "Filter")
+	TArray<FString> IncludePathWildcards;
+
+	/* Packages whose paths match these wildcard patterns will be excluded from gathering. */
+	UPROPERTY(config, EditAnywhere, Category = "Filter")
+	TArray<FString> ExcludePathWildcards;
+
+	/* Packages whose names match these wildcard patterns may be processed for text to gather. */
+	UPROPERTY(config, EditAnywhere, Category = "Filter")
+	TArray<FString> FileExtensions;
+};
+
+USTRUCT()
+struct FMetaDataKeyGatherSpecification
+{
+	GENERATED_USTRUCT_BODY()
+
+	/*  The metadata key for which values will be gathered as text. */
+	UPROPERTY(config, EditAnywhere, Category = "Input")
+	FString MetaDataKey;
+
+	/* The localization namespace in which the gathered text will be output. */
+	UPROPERTY(config, EditAnywhere, Category = "Output")
+	FString TextNamespace;
+
+	/* The pattern which will be formatted to form the localization key for the metadata value gathered as text.
+	Placeholder - Description
+	{FieldPath} - The fully qualified name of the object upon which the metadata resides.
+	{MetaDataValue} - The value associated with the metadata key. */
+	UPROPERTY(config, EditAnywhere, Category = "Output")
+	FString TextKeyPattern;
+};
+
+USTRUCT()
+struct FGatherTextFromMetaDataConfiguration
+{
+	GENERATED_USTRUCT_BODY()
+
+	/* Metadata from source files whose paths match these wildcard patterns, specified relative to the project's root, may be processed for gathering. */
+	UPROPERTY(config, EditAnywhere, Category = "Filter")
+	TArray<FString> IncludePathWildcards;
+
+	/* Metadata from source files whose paths match these wildcard patterns will be excluded from gathering. */
+	UPROPERTY(config, EditAnywhere, Category = "Filter")
+	TArray<FString> ExcludePathWildcards;
+
+	/* Specifications for how to gather text from specific metadata keys. */
+	UPROPERTY(config, EditAnywhere, Category = "MetaData")
+	TArray<FMetaDataKeyGatherSpecification> KeySpecifications;
+};
+
+USTRUCT()
+struct FCultureStatistics
+{
+	GENERATED_USTRUCT_BODY()
+
+	FCultureStatistics()
+		: WordCount(0)
+	{
+	}
+
+	/* The ISO name for this culture. */
+	UPROPERTY(config, EditAnywhere, Category = "Culture")
+	FString CultureName;
+
+	/* The estimated number of words that have been localized for this culture. */
+	UPROPERTY(Transient, EditAnywhere, Category = "Statistics")
+	uint32 WordCount;
+};
+
+UENUM()
+enum class ELocalizationTargetLoadingPolicy : uint8
+{
+	Never,			/* This target's localization data will never be loaded automatically. */
+	Always,			/* This target's localization data will always be loaded automatically. */
+	Editor,			/* This target's localization data will only be loaded when running the editor. Use if this target localizes the editor. */
+	Game,			/* This target's localization data will only be loaded when running the game. Use if this target localizes your game. */
+	PropertyNames,	/* This target's localization data will only be loaded if the editor is displaying localized property names. */
+	ToolTips,		/* This target's localization data will only be loaded if the editor is displaying localized tool tips. */
+};
+
+USTRUCT()
+struct FLocalizationTargetSettings
+{
+	GENERATED_USTRUCT_BODY()
+
+	FLocalizationTargetSettings();
+
+	/* Unique name for the target. */
+	UPROPERTY(config, EditAnywhere, Category = "Target")
+	FString Name;
+
+	UPROPERTY(config)
+	FGuid Guid;
+
+	/* Whether the target has outstanding conflicts that require resolution. */
+	UPROPERTY(Transient, EditAnywhere, Category = "Target")
+	ELocalizationTargetConflictStatus ConflictStatus;
+
+	/* Text present in these targets will not be duplicated in this target. */
+	UPROPERTY(config, EditAnywhere, Category = "Gather")
+	TArray<FGuid> TargetDependencies;
+
+	/* Text present in these manifests will not be duplicated in this target. */
+	UPROPERTY(config, EditAnywhere, Category = "Gather", AdvancedDisplay, meta=(FilePathFilter="manifest"))
+	TArray<FFilePath> AdditionalManifestDependencies;
+
+	/* The names of modules which must be loaded when gathering text. Use to gather from packages or metadata sourced from a module that isn't the primary game module. */
+	UPROPERTY(config, EditAnywhere, Category = "Gather", AdvancedDisplay)
+	TArray<FString> RequiredModuleNames;
+
+	/* Parameters for defining what text is gathered from text files and how. */
+	UPROPERTY(config, EditAnywhere, Category = "Gather")
+	FGatherTextFromTextFilesConfiguration GatherFromTextFiles;
+
+	/* Parameters for defining what text is gathered from packages and how. */
+	UPROPERTY(config, EditAnywhere, Category = "Gather")
+	FGatherTextFromPackagesConfiguration GatherFromPackages;
+
+	/* Parameters for defining what text is gathered from metadata and how. */
+	UPROPERTY(config, EditAnywhere, Category = "Gather")
+	FGatherTextFromMetaDataConfiguration GatherFromMetaData;
+
+	/* The culture in which the source text is written in. */
+	UPROPERTY(config, EditAnywhere, Category = "Cultures")
+	FCultureStatistics NativeCultureStatistics;
+
+	/* Cultures for which the source text is being localized for.*/
+	UPROPERTY(config, EditAnywhere, Category = "Cultures")
+	TArray<FCultureStatistics> SupportedCulturesStatistics;
+};
+
+UCLASS(Within=LocalizationTargetSet)
+class ULocalizationTarget : public UObject
+{
+	GENERATED_BODY()
+
+public:
+	UPROPERTY(EditAnywhere, Category = "Target")
+	FLocalizationTargetSettings Settings;
+
+public:
+#if WITH_EDITOR
+	virtual void PostEditChangeProperty(FPropertyChangedEvent& PropertyChangedEvent) override;
+#endif
+
+	bool IsMemberOfEngineTargetSet() const;
+	bool UpdateWordCountsFromCSV();
+	void UpdateStatusFromConflictReport();
+	bool RenameTargetAndFiles(const FString& NewName);
+	bool DeleteFiles(const FString* const Culture = nullptr) const;
+};
+
+UCLASS(Within=LocalizationDashboardSettings)
+class ULocalizationTargetSet : public UObject
+{
+	GENERATED_BODY()
+
+public:
+	UPROPERTY(EditAnywhere, Category = "Targets")
+	TArray<ULocalizationTarget*> TargetObjects;
+
+public:
+#if WITH_EDITOR
+	virtual void PostEditChangeProperty(FPropertyChangedEvent& PropertyChangedEvent) override;
+#endif
+};
\ No newline at end of file
diff --git a/Engine/Source/Editor/LocalizationDashboard/Private/SLocalizationDashboardTargetRow.h b/Engine/Source/Editor/LocalizationDashboard/Private/SLocalizationDashboardTargetRow.h
--- a/Engine/Source/Editor/LocalizationDashboard/Private/SLocalizationDashboardTargetRow.h
+++ b/Engine/Source/Editor/LocalizationDashboard/Private/SLocalizationDashboardTargetRow.h
@@ -4,8 +4,8 @@
 #include "STableRow.h"
 #include "IPropertyUtilities.h"
 #include "PropertyHandle.h"
-#include "LocalizationTarget.h"
-#include "LocalizationTarget.h"
+#include "LocalizationTargetTypes.h"
+#include "LocalizationTargetTypes.h"
 #include "SDockTab.h"
 
 class SLocalizationDashboardTargetRow : public SMultiColumnTableRow< TSharedPtr<IPropertyHandle> >
diff --git a/Engine/Source/Editor/LocalizationDashboard/Private/SLocalizationTargetEditor.cpp b/Engine/Source/Editor/LocalizationDashboard/Private/SLocalizationTargetEditor.cpp
--- a/Engine/Source/Editor/LocalizationDashboard/Private/SLocalizationTargetEditor.cpp
+++ b/Engine/Source/Editor/LocalizationDashboard/Private/SLocalizationTargetEditor.cpp
@@ -3,7 +3,7 @@
 #include "LocalizationDashboardPrivatePCH.h"
 #include "SLocalizationTargetEditor.h"
 #include "IDetailsView.h"
-#include "LocalizationTarget.h"
+#include "LocalizationTargetTypes.h"
 #include "PropertyEditorModule.h"
 
 void SLocalizationTargetEditor::Construct(const FArguments& InArgs, ULocalizationTargetSet* const InProjectSettings, ULocalizationTarget* const InLocalizationTarget)
diff --git a/Engine/Source/Editor/LocalizationDashboard/Private/SLocalizationTargetEditorCultureRow.h b/Engine/Source/Editor/LocalizationDashboard/Private/SLocalizationTargetEditorCultureRow.h
--- a/Engine/Source/Editor/LocalizationDashboard/Private/SLocalizationTargetEditorCultureRow.h
+++ b/Engine/Source/Editor/LocalizationDashboard/Private/SLocalizationTargetEditorCultureRow.h
@@ -4,7 +4,7 @@
 #include "STableRow.h"
 #include "IPropertyUtilities.h"
 #include "PropertyHandle.h"
-#include "LocalizationTarget.h"
+#include "LocalizationTargetTypes.h"
 
 class SLocalizationTargetEditorCultureRow : public SMultiColumnTableRow<FCulturePtr>
 {
diff --git a/Engine/Source/Editor/LocalizationDashboard/Private/SLocalizationTargetStatusButton.h b/Engine/Source/Editor/LocalizationDashboard/Private/SLocalizationTargetStatusButton.h
--- a/Engine/Source/Editor/LocalizationDashboard/Private/SLocalizationTargetStatusButton.h
+++ b/Engine/Source/Editor/LocalizationDashboard/Private/SLocalizationTargetStatusButton.h
@@ -2,7 +2,7 @@
 #pragma once
 
 #include "SCompoundWidget.h"
-#include "LocalizationTarget.h"
+#include "LocalizationTargetTypes.h"
 #include "LocalizationConfigurationScript.h"
 
 class SLocalizationTargetStatusButton : public SButton
diff --git a/Engine/Source/Editor/LocalizationDashboard/Public/ILocalizationServiceProvider.h b/Engine/Source/Editor/LocalizationDashboard/Public/ILocalizationServiceProvider.h
--- a/Engine/Source/Editor/LocalizationDashboard/Public/ILocalizationServiceProvider.h
+++ b/Engine/Source/Editor/LocalizationDashboard/Public/ILocalizationServiceProvider.h
@@ -6,7 +6,7 @@
 #include "EditorStyleSet.h"
 #include "PropertyHandle.h"
 #include "DetailCategoryBuilder.h"
-#include "LocalizationTarget.h"
+#include "LocalizationTargetTypes.h"
 
 /**
  * Interface to talking with localization service providers.
