Removed unnecessary public include path in LocalizationDashboard.Build.CS.

diff --git a/Engine/Source/Editor/LocalizationDashboard/LocalizationDashboard.Build.cs b/Engine/Source/Editor/LocalizationDashboard/LocalizationDashboard.Build.cs
--- a/Engine/Source/Editor/LocalizationDashboard/LocalizationDashboard.Build.cs
+++ b/Engine/Source/Editor/LocalizationDashboard/LocalizationDashboard.Build.cs
@@ -42,7 +42,6 @@ public class LocalizationDashboard : ModuleRules
 			new string[]
 			{
 				"Editor/LocalizationDashboard/Public",
-				"Editor/LocalizationDashboard/Classes",
 			}
 		);
 	}
