Added missing localization dashboard settings header and source file.

diff --git a/Engine/Source/Editor/LocalizationDashboard/Private/LocalizationDashboardSettings.cpp b/Engine/Source/Editor/LocalizationDashboard/Private/LocalizationDashboardSettings.cpp
new file mode 100644
--- /dev/null
+++ b/Engine/Source/Editor/LocalizationDashboard/Private/LocalizationDashboardSettings.cpp
@@ -0,0 +1,109 @@
+// Copyright 1998-2015 Epic Games, Inc. All Rights Reserved.
+
+#include "LocalizationDashboardPrivatePCH.h"
+#include "LocalizationDashboardSettings.h"
+#include "LocalizationTargetTypes.h"
+#include "ISourceControlModule.h"
+
+ULocalizationDashboardSettings::ULocalizationDashboardSettings(const FObjectInitializer& ObjectInitializer)
+	: Super(ObjectInitializer)
+	, EngineTargetSet(ObjectInitializer.CreateDefaultSubobject<ULocalizationTargetSet>(this, TEXT("EngineLocalizationTargetSet")))
+	, GameTargetSet(ObjectInitializer.CreateDefaultSubobject<ULocalizationTargetSet>(this, TEXT("ProjectLocalizationTargetSet")))
+{	
+}
+
+#if WITH_EDITOR
+void ULocalizationDashboardSettings::PostInitProperties()
+{
+	Super::PostInitProperties();
+
+	// Create and initialize objects for details model from backing config properties.
+	if (EngineTargetSet)
+	{
+		EngineTargetSet->TargetObjects.Empty(EngineTargetsSettings.Num());
+		for (const auto& TargetSettings : EngineTargetsSettings)
+		{
+			ULocalizationTarget* const TargetObject = NewObject<ULocalizationTarget>(EngineTargetSet);
+			TargetObject->Settings = TargetSettings;
+			TargetObject->UpdateStatusFromConflictReport();
+			TargetObject->UpdateWordCountsFromCSV();
+			EngineTargetSet->TargetObjects.Add(TargetObject);
+		}
+	}
+
+	// Create and initialize objects for details model from backing config properties.
+	if (GameTargetSet)
+	{
+		GameTargetSet->TargetObjects.Empty(EngineTargetsSettings.Num());
+		for (const auto& TargetSettings : GameTargetsSettings)
+		{
+			ULocalizationTarget* const TargetObject = NewObject<ULocalizationTarget>(GameTargetSet);
+			TargetObject->Settings = TargetSettings;
+			TargetObject->UpdateStatusFromConflictReport();
+			TargetObject->UpdateWordCountsFromCSV();
+			GameTargetSet->TargetObjects.Add(TargetObject);
+		}
+	}
+}
+
+void ULocalizationDashboardSettings::PostEditChangeProperty(FPropertyChangedEvent& PropertyChangedEvent)
+{
+	Super::PostEditChangeProperty(PropertyChangedEvent);
+
+	// Copy settings back.
+	if (EngineTargetSet)
+	{
+		EngineTargetsSettings.Empty(EngineTargetSet->TargetObjects.Num());
+		for (const auto& TargetObject : EngineTargetSet->TargetObjects)
+		{
+			EngineTargetsSettings.Add(TargetObject ? TargetObject->Settings : FLocalizationTargetSettings());
+		}
+	}
+
+	// Copy settings back.
+	if (GameTargetSet)
+	{
+		GameTargetsSettings.Empty(GameTargetSet->TargetObjects.Num());
+		for (const auto& TargetObject : GameTargetSet->TargetObjects)
+		{
+			GameTargetsSettings.Add(TargetObject ? TargetObject->Settings : FLocalizationTargetSettings());
+		}
+	}
+
+	// Handle source control operations before saving to config.
+	ISourceControlModule& SourceControl = ISourceControlModule::Get();
+	ISourceControlProvider& SourceControlProvider = SourceControl.GetProvider();
+	const bool CanUseSourceControl = SourceControl.IsEnabled() && SourceControlProvider.IsEnabled() && SourceControlProvider.IsAvailable();
+
+	if (CanUseSourceControl)
+	{
+		const FString Path = GetDefaultConfigFilename();
+		const FSourceControlStatePtr SourceControlState = SourceControlProvider.GetState(*Path, EStateCacheUsage::Use);
+
+		// Should only use source control if the file is actually under source control.
+		if (CanUseSourceControl && SourceControlState.IsValid() && !SourceControlState->CanAdd())
+		{
+			if (!SourceControlState->CanEdit())
+			{
+				SourceControlProvider.Execute(ISourceControlOperation::Create<FCheckOut>(), *Path);
+			}
+		}
+	}
+
+	UpdateDefaultConfigFile();
+}
+#endif
+
+ULocalizationTargetSet* ULocalizationDashboardSettings::GetEngineTargetSet()
+{
+	ULocalizationDashboardSettings* LocalizationDashboardSettings = GetMutableDefault<ULocalizationDashboardSettings>();
+	check(LocalizationDashboardSettings);
+	return LocalizationDashboardSettings->EngineTargetSet;
+}
+
+ULocalizationTargetSet* ULocalizationDashboardSettings::GetGameTargetSet()
+{
+	ULocalizationDashboardSettings* LocalizationDashboardSettings = GetMutableDefault<ULocalizationDashboardSettings>();
+	check(LocalizationDashboardSettings);
+	return LocalizationDashboardSettings->GameTargetSet;
+}
\ No newline at end of file
diff --git a/Engine/Source/Editor/LocalizationDashboard/Private/LocalizationDashboardSettings.h b/Engine/Source/Editor/LocalizationDashboard/Private/LocalizationDashboardSettings.h
new file mode 100644
--- /dev/null
+++ b/Engine/Source/Editor/LocalizationDashboard/Private/LocalizationDashboardSettings.h
@@ -0,0 +1,39 @@
+// Copyright 1998-2015 Epic Games, Inc. All Rights Reserved.
+
+#pragma once
+
+#include "Object.h"
+#include "LocalizationTargetTypes.h"
+#include "LocalizationDashboardSettings.generated.h"
+
+// Class for loading/saving configuration settings and the details view objects needed for localization dashboard functionality.
+UCLASS(Config=Editor, defaultconfig)
+class ULocalizationDashboardSettings : public UObject
+{
+	GENERATED_BODY()
+
+public:
+	ULocalizationDashboardSettings(const FObjectInitializer& ObjectInitializer);
+
+private:
+	UPROPERTY()
+	ULocalizationTargetSet* EngineTargetSet;
+
+	UPROPERTY(config)
+	TArray<FLocalizationTargetSettings> EngineTargetsSettings;
+
+	UPROPERTY()
+	ULocalizationTargetSet* GameTargetSet;
+
+	UPROPERTY(config)
+	TArray<FLocalizationTargetSettings> GameTargetsSettings;
+
+public:
+#if WITH_EDITOR
+	virtual void PostInitProperties() override;
+	virtual void PostEditChangeProperty(FPropertyChangedEvent& PropertyChangedEvent) override;
+#endif
+
+	static ULocalizationTargetSet* GetEngineTargetSet();
+	static ULocalizationTargetSet* GetGameTargetSet();
+};
\ No newline at end of file
