Source control operations on the config files for the localization dashboard settings are now handled manually and exposed via a widget in the localization dashboard.

diff --git a/Engine/Source/Editor/LocalizationDashboard/LocalizationDashboard.Build.cs b/Engine/Source/Editor/LocalizationDashboard/LocalizationDashboard.Build.cs
--- a/Engine/Source/Editor/LocalizationDashboard/LocalizationDashboard.Build.cs
+++ b/Engine/Source/Editor/LocalizationDashboard/LocalizationDashboard.Build.cs
@@ -27,7 +27,8 @@ public class LocalizationDashboard : ModuleRules
 				"DesktopPlatform",
                 "TranslationEditor",
                 "MainFrame",
-                "SourceControl"
+                "SourceControl",
+                "SharedSettingsWidgets",
 			}
 		);
 
diff --git a/Engine/Source/Editor/LocalizationDashboard/Private/LocalizationDashboard.cpp b/Engine/Source/Editor/LocalizationDashboard/Private/LocalizationDashboard.cpp
--- a/Engine/Source/Editor/LocalizationDashboard/Private/LocalizationDashboard.cpp
+++ b/Engine/Source/Editor/LocalizationDashboard/Private/LocalizationDashboard.cpp
@@ -10,6 +10,7 @@
 #include "SLocalizationTargetEditor.h"
 #include "LocalizationDashboardSettings.h"
 #include "LocalizationTargetTypes.h"
+#include "SSettingsEditorCheckoutNotice.h"
 
 #define LOCTEXT_NAMESPACE "LocalizationDashboard"
 
@@ -25,6 +26,9 @@ public:
 
 	TWeakPtr<SDockTab> ShowTargetEditor(ULocalizationTarget* const LocalizationTarget);
 
+private:
+	bool CanMakeEdits() const;
+
 private:
 	static const FName EngineTargetSetDetailsTabName;
 	static const FName ProjectTargetSetDetailsTabName;
@@ -32,6 +36,7 @@ private:
 
 	TSharedPtr<FTabManager> TabManager;
 	TMap< TWeakObjectPtr<ULocalizationTarget>, TWeakPtr<SDockTab> > TargetToTabMap;
+	TSharedPtr<SSettingsEditorCheckoutNotice> SettingsEditorCheckoutNotice;
 };
 
 const FName SLocalizationDashboard::EngineTargetSetDetailsTabName("EngineTargets");
@@ -48,7 +53,7 @@ void SLocalizationDashboard::Construct(const FArguments& InArgs, const TSharedPt
 	};
 	TabManager->SetOnPersistLayout(FTabManager::FOnPersistLayout::CreateLambda(PersistLayout));
 
-	const auto& CreateDetailsTab = [](const FSpawnTabArgs& SpawnTabArgs) -> TSharedRef<SDockTab>
+	const auto& CreateDetailsTab = [this](const FSpawnTabArgs& SpawnTabArgs) -> TSharedRef<SDockTab>
 	{
 		FText DockTabLabel;
 		ULocalizationTargetSet* TargetSet = nullptr;
@@ -69,6 +74,7 @@ void SLocalizationDashboard::Construct(const FArguments& InArgs, const TSharedPt
 		FPropertyEditorModule& PropertyModule = FModuleManager::LoadModuleChecked<FPropertyEditorModule>("PropertyEditor");
 		FDetailsViewArgs DetailsViewArgs(false, false, false, FDetailsViewArgs::ENameAreaSettings::HideNameArea, false, nullptr, false, NAME_None);
 		TSharedRef<IDetailsView> DetailsView = PropertyModule.CreateDetailView(DetailsViewArgs);
+		DetailsView->SetIsPropertyEditingEnabledDelegate(FIsPropertyEditingEnabled::CreateSP(this, &SLocalizationDashboard::CanMakeEdits));
 
 		DetailsView->SetObject(TargetSet, true);
 
@@ -110,6 +116,10 @@ void SLocalizationDashboard::Construct(const FArguments& InArgs, const TSharedPt
 	IMainFrameModule& MainFrameModule = FModuleManager::LoadModuleChecked<IMainFrameModule>( "MainFrame" );
 	const TSharedRef<SWidget> MenuWidget = MainFrameModule.MakeMainMenu( TabManager, MenuExtender );
 
+	FString ConfigFilePath;
+	ConfigFilePath = GetDefault<ULocalizationDashboardSettings>()->GetDefaultConfigFilename();
+	ConfigFilePath = FPaths::ConvertRelativePathToFull(ConfigFilePath);
+
 	ChildSlot
 		[
 			SNew(SVerticalBox)
@@ -119,9 +129,16 @@ void SLocalizationDashboard::Construct(const FArguments& InArgs, const TSharedPt
 				MenuWidget
 			]
 			+SVerticalBox::Slot()
+			.AutoHeight()
+			.Padding(0.0f, 8.0f)
+			[
+				SAssignNew(SettingsEditorCheckoutNotice, SSettingsEditorCheckoutNotice)
+				.ConfigFilePath(ConfigFilePath)
+			]
+			+SVerticalBox::Slot()
 			.FillHeight(1.0f)
 			[
-			TabManager->RestoreFrom(Layout, OwningWindow).ToSharedRef()
+				TabManager->RestoreFrom(Layout, OwningWindow).ToSharedRef()
 			]
 		];
 }
@@ -135,7 +152,7 @@ TWeakPtr<SDockTab> SLocalizationDashboard::ShowTargetEditor(ULocalizationTarget*
 	{
 		ULocalizationTargetSet* const TargetSet = LocalizationTarget->GetTypedOuter<ULocalizationTargetSet>();
 
-		const TSharedRef<SLocalizationTargetEditor> OurTargetEditor = SNew(SLocalizationTargetEditor, TargetSet, LocalizationTarget);
+		const TSharedRef<SLocalizationTargetEditor> OurTargetEditor = SNew(SLocalizationTargetEditor, TargetSet, LocalizationTarget, FIsPropertyEditingEnabled::CreateSP(this, &SLocalizationDashboard::CanMakeEdits));
 		const TSharedRef<SDockTab> NewTargetEditorTab = SNew(SDockTab)
 			.TabRole(ETabRole::DocumentTab)
 			.Label_Lambda( [LocalizationTarget]
@@ -158,6 +175,11 @@ TWeakPtr<SDockTab> SLocalizationDashboard::ShowTargetEditor(ULocalizationTarget*
 	return TargetEditorDockTab;
 }
 
+bool SLocalizationDashboard::CanMakeEdits() const
+{
+	return SettingsEditorCheckoutNotice.IsValid() && SettingsEditorCheckoutNotice->IsUnlocked();
+}
+
 FLocalizationDashboard* FLocalizationDashboard::Instance = nullptr;
 
 FLocalizationDashboard* FLocalizationDashboard::Get()
diff --git a/Engine/Source/Editor/LocalizationDashboard/Private/LocalizationDashboardSettings.cpp b/Engine/Source/Editor/LocalizationDashboard/Private/LocalizationDashboardSettings.cpp
--- a/Engine/Source/Editor/LocalizationDashboard/Private/LocalizationDashboardSettings.cpp
+++ b/Engine/Source/Editor/LocalizationDashboard/Private/LocalizationDashboardSettings.cpp
@@ -70,26 +70,6 @@ void ULocalizationDashboardSettings::PostEditChangeProperty(FPropertyChangedEven
 		}
 	}
 
-	// Handle source control operations before saving to config.
-	ISourceControlModule& SourceControl = ISourceControlModule::Get();
-	ISourceControlProvider& SourceControlProvider = SourceControl.GetProvider();
-	const bool CanUseSourceControl = SourceControl.IsEnabled() && SourceControlProvider.IsEnabled() && SourceControlProvider.IsAvailable();
-
-	if (CanUseSourceControl)
-	{
-		const FString Path = GetDefaultConfigFilename();
-		const FSourceControlStatePtr SourceControlState = SourceControlProvider.GetState(*Path, EStateCacheUsage::Use);
-
-		// Should only use source control if the file is actually under source control.
-		if (CanUseSourceControl && SourceControlState.IsValid() && !SourceControlState->CanAdd())
-		{
-			if (!SourceControlState->CanEdit())
-			{
-				SourceControlProvider.Execute(ISourceControlOperation::Create<FCheckOut>(), *Path);
-			}
-		}
-	}
-
 	UpdateDefaultConfigFile();
 }
 #endif
diff --git a/Engine/Source/Editor/LocalizationDashboard/Private/SLocalizationTargetEditor.cpp b/Engine/Source/Editor/LocalizationDashboard/Private/SLocalizationTargetEditor.cpp
--- a/Engine/Source/Editor/LocalizationDashboard/Private/SLocalizationTargetEditor.cpp
+++ b/Engine/Source/Editor/LocalizationDashboard/Private/SLocalizationTargetEditor.cpp
@@ -6,13 +6,14 @@
 #include "LocalizationTargetTypes.h"
 #include "PropertyEditorModule.h"
 
-void SLocalizationTargetEditor::Construct(const FArguments& InArgs, ULocalizationTargetSet* const InProjectSettings, ULocalizationTarget* const InLocalizationTarget)
+void SLocalizationTargetEditor::Construct(const FArguments& InArgs, ULocalizationTargetSet* const InProjectSettings, ULocalizationTarget* const InLocalizationTarget, const FIsPropertyEditingEnabled& IsPropertyEditingEnabled)
 {
 	check(InProjectSettings->TargetObjects.Contains(InLocalizationTarget));
 
 	FPropertyEditorModule& PropertyModule = FModuleManager::LoadModuleChecked<FPropertyEditorModule>("PropertyEditor");
 	FDetailsViewArgs DetailsViewArgs(false, false, false, FDetailsViewArgs::ENameAreaSettings::HideNameArea, false, nullptr, false, NAME_None);
 	TSharedRef<IDetailsView> DetailsView = PropertyModule.CreateDetailView(DetailsViewArgs);
+	DetailsView->SetIsPropertyEditingEnabledDelegate(IsPropertyEditingEnabled);
 
 	ChildSlot
 		[
diff --git a/Engine/Source/Editor/LocalizationDashboard/Private/SLocalizationTargetEditor.h b/Engine/Source/Editor/LocalizationDashboard/Private/SLocalizationTargetEditor.h
--- a/Engine/Source/Editor/LocalizationDashboard/Private/SLocalizationTargetEditor.h
+++ b/Engine/Source/Editor/LocalizationDashboard/Private/SLocalizationTargetEditor.h
@@ -13,6 +13,6 @@ public:
 	SLATE_BEGIN_ARGS(SLocalizationTargetEditor) {}
 	SLATE_END_ARGS()
 
-	void Construct(const FArguments& InArgs, ULocalizationTargetSet* const InProjectSettings, ULocalizationTarget* const InTarget);
+	void Construct(const FArguments& InArgs, ULocalizationTargetSet* const InProjectSettings, ULocalizationTarget* const InTarget, const FIsPropertyEditingEnabled& IsPropertyEditingEnabled);
 };
 
