Renamed LocalizationTarget.cpp to LocalizationTargetTypes.cpp for consistency and descriptiveness.

diff --git a/Engine/Source/Editor/LocalizationDashboard/Private/LocalizationTarget.cpp b/Engine/Source/Editor/LocalizationDashboard/Private/LocalizationTarget.cpp
deleted file mode 100644
diff --git a/Engine/Source/Editor/LocalizationDashboard/Private/LocalizationTargetTypes.cpp b/Engine/Source/Editor/LocalizationDashboard/Private/LocalizationTargetTypes.cpp
new file mode 100644
--- /dev/null
+++ b/Engine/Source/Editor/LocalizationDashboard/Private/LocalizationTargetTypes.cpp
@@ -0,0 +1,462 @@
+﻿// Copyright 1998-2015 Epic Games, Inc. All Rights Reserved.
+
+#include "LocalizationDashboardPrivatePCH.h"
+#include "LocalizationTargetTypes.h"
+#include "CsvParser.h"
+#include "LocalizationConfigurationScript.h"
+#include "LocalizationDashboardSettings.h"
+
+FLocalizationTargetSettings::FLocalizationTargetSettings()
+	: Guid(FGuid::NewGuid())
+	, ConflictStatus(ELocalizationTargetConflictStatus::Unknown)
+{
+}
+
+namespace
+{
+	class FWordCountCSVParser
+	{
+		typedef TArray<const TCHAR*> FRow;
+
+	public:
+		static bool Execute(FLocalizationTargetSettings& Target, const FString& CSVFilePath)
+		{
+			// Load CSV file to string.
+			FString CSVString;
+			FFileHelper::LoadFileToString(CSVString, *CSVFilePath);
+
+			// Parse CSV file string.
+			FCsvParser CSVParser(CSVString);
+			const FCsvParser::FRows& CSVRows = CSVParser.GetRows();
+
+			// Initialize.
+			FWordCountCSVParser Instance(Target);
+
+			// Parse header row.
+			if ( !(CSVRows.Num() > 0) ||!Instance.ParseHeaderRow(CSVRows[0]) )
+			{
+				return false;
+			}
+
+			// Parse latest word count row.
+			TMap<FString, uint32> WordCounts;
+			if ( !Instance.ParseWordCountRow(CSVRows.Top(), WordCounts) )
+			{
+				return false;
+			}
+
+			uint32* NativeWordCount = WordCounts.Find(Target.NativeCultureStatistics.CultureName);
+			if (NativeWordCount)
+			{
+				Target.NativeCultureStatistics.WordCount = *NativeWordCount;
+			}
+
+			for (FCultureStatistics& SupportedCultureStatistics : Target.SupportedCulturesStatistics)
+			{
+				uint32* SupportedWordCount = WordCounts.Find(SupportedCultureStatistics.CultureName);
+				if (SupportedWordCount)
+				{
+					SupportedCultureStatistics.WordCount = *SupportedWordCount;
+				}
+			}
+
+			return true;
+		}
+
+	private:
+		FWordCountCSVParser(FLocalizationTargetSettings& InTarget)
+			: Target(InTarget)
+		{
+			MandatoryColumnNames.Add( TEXT("Date/Time") );
+			MandatoryColumnNames.Add( TEXT("Word Count") );
+		}
+
+		bool ParseHeaderRow(const FRow& HeaderRow)
+		{
+			// Validate header row.
+			const bool IsHeaderRowValid = [&]() -> bool
+			{
+				if (HeaderRow.Num() < MandatoryColumnNames.Num())
+				{
+					// Not enough columns present to represent the mandatory columns.
+					return false;
+				}
+				for (int32 i = 0; i < MandatoryColumnNames.Num(); ++i)
+				{
+					if (MandatoryColumnNames[i] != HeaderRow[i])
+					{
+						return false;
+					}
+				}
+
+				return true;
+			}();
+
+			if (!IsHeaderRowValid)
+			{
+				return false;
+			}
+
+			// Parse culture names.
+			for (int32 i = MandatoryColumnNames.Num(); i < HeaderRow.Num(); ++i)
+			{
+				CultureNames.Add(HeaderRow[i]);
+			}
+
+			return true;
+		}
+
+		bool ParseWordCountRow(const FRow& WordCountRow, TMap<FString, uint32>& OutCultureWordCounts)
+		{
+			// Validate row.
+			bool IsRowValid = [&]() -> bool
+			{
+				if (WordCountRow.Num() != MandatoryColumnNames.Num() + CultureNames.Num())
+				{
+					return false;
+				}
+
+				return true;
+			}();
+
+			if (!IsRowValid)
+			{
+				return false;
+			}
+
+			// Parse time stamp.
+			FDateTime TimeStamp;
+			if (!FDateTime::Parse(WordCountRow[0], TimeStamp))
+			{
+				return false;
+			}
+
+			// Parse native culture word count.
+			{
+				const uint32 WordCount = static_cast<uint32>( FCString::Atoi(WordCountRow[1]) );
+				OutCultureWordCounts.Add(Target.NativeCultureStatistics.CultureName, WordCount);
+			}
+
+			// Parse supported culture word counts.
+			for (int32 j = MandatoryColumnNames.Num(); j < WordCountRow.Num(); ++j)
+			{
+				const uint32 WordCount = static_cast<uint32>( FCString::Atoi(WordCountRow[j]) );
+				const int32 CultureIndex = j - MandatoryColumnNames.Num();
+				OutCultureWordCounts.FindOrAdd(CultureNames[CultureIndex]) = WordCount;
+			}
+
+			return true;
+		}
+
+	private:
+		FLocalizationTargetSettings& Target;
+		TArray<FString> MandatoryColumnNames;
+		TArray<FString> CultureNames;
+	};
+}
+
+void ULocalizationTarget::PostEditChangeProperty(FPropertyChangedEvent& PropertyChangedEvent)
+{
+	Super::PostEditChangeProperty(PropertyChangedEvent);
+
+	// Notify parent of change.
+	ULocalizationTargetSet* const LocalizationTargetSet = Cast<ULocalizationTargetSet>(GetOuter());
+	if (LocalizationTargetSet)
+	{
+		LocalizationTargetSet->PostEditChange();
+	}
+}
+
+bool ULocalizationTarget::IsMemberOfEngineTargetSet() const
+{
+	ULocalizationTargetSet* const TargetSet = CastChecked<ULocalizationTargetSet>(GetOuter());
+	return TargetSet == ULocalizationDashboardSettings::GetEngineTargetSet();
+}
+
+bool ULocalizationTarget::UpdateWordCountsFromCSV()
+{
+	const FString CSVFilePath = LocalizationConfigurationScript::GetWordCountCSVPath(this);
+	return FWordCountCSVParser::Execute(Settings, CSVFilePath);
+}
+
+void ULocalizationTarget::UpdateStatusFromConflictReport()
+{
+	const FString ConflictReportFilePath = LocalizationConfigurationScript::GetConflictReportPath(this);
+	const int64 ConflictReportFileSize = IFileManager::Get().FileSize(*ConflictReportFilePath);
+
+	switch (ConflictReportFileSize)
+	{
+	case INDEX_NONE:
+		// Unknown
+		Settings.ConflictStatus = ELocalizationTargetConflictStatus::Unknown;
+		break;
+	case 0:
+		// Good
+		Settings.ConflictStatus = ELocalizationTargetConflictStatus::Clear;
+		break;
+	default:
+		// Bad
+		Settings.ConflictStatus = ELocalizationTargetConflictStatus::ConflictsPresent;
+		break;
+	}
+}
+
+bool ULocalizationTarget::RenameTargetAndFiles(const FString& NewName)
+{
+	bool HasCompletelySucceeded = true;
+	ISourceControlModule& SourceControl = ISourceControlModule::Get();
+	ISourceControlProvider& SourceControlProvider = SourceControl.GetProvider();
+	const bool CanUseSourceControl = SourceControl.IsEnabled() && SourceControlProvider.IsEnabled() && SourceControlProvider.IsAvailable();
+
+	const auto& TryDelete = [&](const FString& Path) -> bool
+	{
+		const FSourceControlStatePtr SourceControlState = SourceControlProvider.GetState(*Path, EStateCacheUsage::Use);
+
+		// Should only use source control if the file is actually under source control.
+		if (CanUseSourceControl && SourceControlState.IsValid() && !SourceControlState->CanAdd())
+		{
+			// File is already marked for deletion.
+			if (SourceControlState->IsDeleted())
+			{
+				return true;
+			}
+			else 
+			{
+				// File is in some modified source control state, needs to be reverted.
+				if (SourceControlState->IsAdded() || SourceControlState->IsCheckedOut())
+				{
+					if (SourceControlProvider.Execute(ISourceControlOperation::Create<FRevert>(), *Path) != ECommandResult::Succeeded)
+					{
+						return false;
+					}
+				}
+
+				// File needs to be deleted.
+				if (SourceControlProvider.Execute(ISourceControlOperation::Create<FDelete>(), *Path) == ECommandResult::Succeeded)
+				{
+					return true;
+				}
+			}
+		}
+		// Attempt local deletion.
+		else if (IFileManager::Get().Delete(*Path, false, true))
+		{
+			return true;
+		}
+
+		return false;
+	};
+
+	// Delete configuration files.
+	const TArray<FString>& ScriptPaths = LocalizationConfigurationScript::GetScriptPaths(this);
+	for (const FString& ScriptPath : ScriptPaths)
+	{
+		if (!TryDelete(ScriptPath))
+		{
+			HasCompletelySucceeded = false;
+		}
+	}
+
+	const auto& GetNamedPaths = [this](TArray<FString>& NamedPaths)
+	{
+		NamedPaths.Add(LocalizationConfigurationScript::GetManifestPath(this));
+		NamedPaths.Add(LocalizationConfigurationScript::GetWordCountCSVPath(this));
+		NamedPaths.Add(LocalizationConfigurationScript::GetConflictReportPath(this));
+
+		TArray<const FCultureStatistics*> Cultures;
+		Cultures.Add(&Settings.NativeCultureStatistics);
+		for (const FCultureStatistics& Culture : Settings.SupportedCulturesStatistics)
+		{
+			Cultures.Add(&Culture);
+		}
+
+		for (const FCultureStatistics* Culture : Cultures)
+		{
+			NamedPaths.Add(LocalizationConfigurationScript::GetArchivePath(this, Culture->CultureName));
+			NamedPaths.Add(LocalizationConfigurationScript::GetDefaultPOPath(this, Culture->CultureName));
+			NamedPaths.Add(LocalizationConfigurationScript::GetLocResPath(this, Culture->CultureName));
+		}
+
+		NamedPaths.Add(LocalizationConfigurationScript::GetDataDirectory(this));
+	};
+
+	TArray<FString> OldPaths;
+	GetNamedPaths(OldPaths);
+
+	// Rename
+	Settings.Name = NewName;
+
+	TArray<FString> NewPaths;
+	GetNamedPaths(NewPaths);
+
+	// Rename data files.
+	check(OldPaths.Num() == NewPaths.Num());
+	for (int32 i = 0; i < OldPaths.Num(); ++i)
+	{
+		if (IFileManager::Get().Move(*NewPaths[i], *OldPaths[i], true, true, false, true))
+		{
+			if (CanUseSourceControl)
+			{
+				// Add new file.
+				if (SourceControlProvider.Execute(ISourceControlOperation::Create<FMarkForAdd>(), *NewPaths[i]) != ECommandResult::Succeeded)
+				{
+					HasCompletelySucceeded = false;
+				}
+			}
+
+			bool ShouldUseSourceControl = false;
+
+			if (CanUseSourceControl)
+			{
+				const FSourceControlStatePtr SourceControlState = SourceControlProvider.GetState(*OldPaths[i], EStateCacheUsage::Use);
+				if (SourceControlState->IsSourceControlled())
+				{
+					ShouldUseSourceControl = true;
+				}
+			}
+
+			// Use source control.
+			if (ShouldUseSourceControl)
+			{
+				// Delete old directory/file.
+				if (SourceControlProvider.Execute(ISourceControlOperation::Create<FDelete>(), *OldPaths[i]) != ECommandResult::Succeeded)
+				{
+					HasCompletelySucceeded = false;
+				}
+			}
+			// Operate locally.
+			else
+			{
+				// Delete old directory/file.
+				if (	(FPaths::DirectoryExists(OldPaths[i]) && !IFileManager::Get().DeleteDirectory(*OldPaths[i], false, true)) ||
+					(FPaths::FileExists(OldPaths[i]) && !IFileManager::Get().Delete(*OldPaths[i], false, true))	)
+				{
+					HasCompletelySucceeded = false;
+				}
+			}
+		}
+	}
+
+	// Generate new configuration files.
+	LocalizationConfigurationScript::GenerateGatherScript(this).Write(LocalizationConfigurationScript::GetGatherScriptPath(this));
+	LocalizationConfigurationScript::GenerateImportScript(this).Write(LocalizationConfigurationScript::GetImportScriptPath(this));
+	LocalizationConfigurationScript::GenerateExportScript(this).Write(LocalizationConfigurationScript::GetExportScriptPath(this));
+	LocalizationConfigurationScript::GenerateWordCountReportScript(this).Write(LocalizationConfigurationScript::GetWordCountReportScriptPath(this));
+
+	//if (CanUseSourceControl)
+	//{
+	//	HasCompletelySucceeded = HasCompletelySucceeded && SourceControlProvider.Execute(ISourceControlOperation::Create<FMarkForAdd>(), *LocalizationConfigurationScript::GetGatherScriptPath(this)) == ECommandResult::Succeeded;
+	//	HasCompletelySucceeded = HasCompletelySucceeded && SourceControlProvider.Execute(ISourceControlOperation::Create<FMarkForAdd>(), *LocalizationConfigurationScript::GetImportScriptPath(this)) == ECommandResult::Succeeded;
+	//	HasCompletelySucceeded = HasCompletelySucceeded && SourceControlProvider.Execute(ISourceControlOperation::Create<FMarkForAdd>(), *LocalizationConfigurationScript::GetExportScriptPath(this)) == ECommandResult::Succeeded;
+	//	HasCompletelySucceeded = HasCompletelySucceeded && SourceControlProvider.Execute(ISourceControlOperation::Create<FMarkForAdd>(), *LocalizationConfigurationScript::GetReportScriptPath(this)) == ECommandResult::Succeeded;
+	//}
+
+	return HasCompletelySucceeded;
+}
+
+bool ULocalizationTarget::DeleteFiles(const FString* const Culture) const
+{
+	bool HasCompletelySucceeded = true;
+	ISourceControlModule& SourceControl = ISourceControlModule::Get();
+	ISourceControlProvider& SourceControlProvider = SourceControl.GetProvider();
+
+	const FString DataDirectory = LocalizationConfigurationScript::GetDataDirectory(this) + (Culture ? *Culture : FString());
+
+	// Remove data files from source control.
+	{
+		TArray<FString> FilesInDataDirectory;
+		IFileManager::Get().FindFilesRecursive(FilesInDataDirectory, *DataDirectory, TEXT("*"), true, false);
+
+		ISourceControlProvider& SourceControlProvider = ISourceControlModule::Get().GetProvider();
+		TArray<FSourceControlStateRef> SourceControlStates;
+		ECommandResult::Type Result = SourceControlProvider.GetState( FilesInDataDirectory, SourceControlStates, EStateCacheUsage::ForceUpdate );
+		if (Result == ECommandResult::Succeeded)
+		{
+			TArray<FString> FilesToRevert;
+			TArray<FString> FilesToDelete;
+
+			for (const FSourceControlStateRef& SourceControlState : SourceControlStates)
+			{
+				// Added/edited files must be reverted.
+				if (SourceControlState->IsAdded() || SourceControlState->IsCheckedOut())
+				{
+					FilesToRevert.Add(SourceControlState->GetFilename());
+				}
+				// Edited files must be deleted.
+				if (SourceControlState->IsCheckedOut())
+				{
+					FilesToDelete.Add(SourceControlState->GetFilename());
+				}
+			}
+
+			const bool WasSuccessfullyReverted = (SourceControlProvider.Execute( ISourceControlOperation::Create<FRevert>(), FilesToRevert ) == ECommandResult::Succeeded);
+			const bool WasSuccessfullyDeleted = (SourceControlProvider.Execute( ISourceControlOperation::Create<FDelete>(), FilesToDelete ) == ECommandResult::Succeeded);
+			HasCompletelySucceeded = HasCompletelySucceeded && WasSuccessfullyReverted && WasSuccessfullyDeleted;
+		}
+	}
+
+	// Delete data files.
+	if (!IFileManager::Get().DeleteDirectory(*DataDirectory, false, true))
+	{
+		HasCompletelySucceeded = false;
+	}
+
+	// Delete configuration files if deleting the target.
+	if (!Culture)
+	{
+		const TArray<FString>& ScriptPaths = LocalizationConfigurationScript::GetScriptPaths(this);
+
+		// Remove script files from source control.
+		{
+			ISourceControlProvider& SourceControlProvider = ISourceControlModule::Get().GetProvider();
+			TArray<FSourceControlStateRef> SourceControlStates;
+			ECommandResult::Type Result = SourceControlProvider.GetState( ScriptPaths, SourceControlStates, EStateCacheUsage::ForceUpdate );
+			if (Result == ECommandResult::Succeeded)
+			{
+				TArray<FString> FilesToRevert;
+				TArray<FString> FilesToDelete;
+
+				for (const FSourceControlStateRef& SourceControlState : SourceControlStates)
+				{
+					// Added/edited files must be reverted.
+					if (SourceControlState->IsAdded() || SourceControlState->IsCheckedOut())
+					{
+						FilesToRevert.Add(SourceControlState->GetFilename());
+					}
+					// Edited files must be deleted.
+					if (SourceControlState->IsCheckedOut())
+					{
+						FilesToDelete.Add(SourceControlState->GetFilename());
+					}
+				}
+
+				const bool WasSuccessfullyReverted = (SourceControlProvider.Execute( ISourceControlOperation::Create<FRevert>(), FilesToRevert ) == ECommandResult::Succeeded);
+				const bool WasSuccessfullyDeleted = (SourceControlProvider.Execute( ISourceControlOperation::Create<FDelete>(), FilesToDelete ) == ECommandResult::Succeeded);
+				HasCompletelySucceeded = HasCompletelySucceeded && WasSuccessfullyReverted && WasSuccessfullyDeleted;
+			}
+		}
+
+		// Delete script files.
+		for (const FString& ScriptPath : ScriptPaths)
+		{
+			if (!IFileManager::Get().Delete(*ScriptPath, false, true))
+			{
+				HasCompletelySucceeded = false;
+			}
+		}
+	}
+
+	return HasCompletelySucceeded;
+}
+
+void ULocalizationTargetSet::PostEditChangeProperty(FPropertyChangedEvent& PropertyChangedEvent)
+{
+	Super::PostEditChangeProperty(PropertyChangedEvent);
+
+	// Notify parent of change.
+	ULocalizationDashboardSettings* const LocalizationDashboardSettings = Cast<ULocalizationDashboardSettings>(GetOuter());
+	if (LocalizationDashboardSettings)
+	{
+		LocalizationDashboardSettings->PostEditChange();
+	}
+}
\ No newline at end of file
