Likely solution for non-Unity-build error regarding FIsPropertyEditingEnabled in SLocalizationTargetEditor.h.

diff --git a/Engine/Source/Editor/LocalizationDashboard/Private/SLocalizationTargetEditor.h b/Engine/Source/Editor/LocalizationDashboard/Private/SLocalizationTargetEditor.h
--- a/Engine/Source/Editor/LocalizationDashboard/Private/SLocalizationTargetEditor.h
+++ b/Engine/Source/Editor/LocalizationDashboard/Private/SLocalizationTargetEditor.h
@@ -3,6 +3,7 @@
 #pragma once
 
 #include "SCompoundWidget.h"
+#include "PropertyEditorDelegates.h"
 
 class ULocalizationTargetSet;
 class ULocalizationTarget;
