Fixing module implementation macro for LocalizationDashboardModule

diff --git a/Engine/Source/Editor/LocalizationDashboard/Private/LocalizationDashboardModule.cpp b/Engine/Source/Editor/LocalizationDashboard/Private/LocalizationDashboardModule.cpp
--- a/Engine/Source/Editor/LocalizationDashboard/Private/LocalizationDashboardModule.cpp
+++ b/Engine/Source/Editor/LocalizationDashboard/Private/LocalizationDashboardModule.cpp
@@ -83,6 +83,6 @@ private:
 	FName CurrentServiceProviderName;
 };
 
-IMPLEMENT_MODULE(FLocalizationDashboardModule, "LocalizationDashboard");
+IMPLEMENT_MODULE(FLocalizationDashboardModule, LocalizationDashboard);
 
 #undef LOCTEXT_NAMESPACE
