Massive overhaul of gather configuration UI in the Localization Dashboard. Updated localization configuration script settings with new, more descriptive names and moved "friendly" behaviors to the Localization Dashboard's localization configuration script generation code, rather than in the commandlets. REVIEW CHANGELIST FOR API CHANGES.

diff --git a/Engine/Config/Localization/Editor.ini b/Engine/Config/Localization/Editor.ini
--- a/Engine/Config/Localization/Editor.ini
+++ b/Engine/Config/Localization/Editor.ini
@@ -1,7 +1,7 @@
 [CommonSettings]
-ManifestDependencies=./Content/Localization/Engine/Engine.manifest
-SourcePath=./Content/Localization/Editor
-DestinationPath=./Content/Localization/Editor
+ManifestDependencies=Content/Localization/Engine/Engine.manifest
+SourcePath=Content/Localization/Editor
+DestinationPath=Content/Localization/Editor
 ManifestName=Editor.manifest
 ArchiveName=Editor.archive
 PortableObjectName=Editor.po
@@ -16,23 +16,23 @@ CulturesToGenerate=zh-CN
 ;Gather text from source code
 [GatherTextStep0]
 CommandletClass=GatherTextFromSource
-IncludePaths=./Source/Editor/
-SourceFileSearchFilters=*.cpp
-SourceFileSearchFilters=*.h
-SourceFileSearchFilters=*.c
-SourceFileSearchFilters=*.inl
-SourceFileSearchFilters=*.mm
+SearchDirectoryPaths=Source/Editor/
+FileNameFilters=*.cpp
+FileNameFilters=*.h
+FileNameFilters=*.c
+FileNameFilters=*.inl
+FileNameFilters=*.mm
 
 ;Gather text from metadata
 [GatherTextStep1]
 CommandletClass=GatherTextFromMetaData
-IncludePaths=./Source/Editor/*
-IncludePaths=./Source/Runtime/*
-IncludePaths=./Source/Developer/*
-ExcludePaths=*/Config/Localization/*
-ExcludePaths=*/Source/Developer/NoRedist/CommunityPortalServices/*
-ExcludePaths=*/Source/Developer/NoRedist/UnrealEngineLauncherServices/*
-ExcludePaths=*/Source/Developer/NoRedist/BuildPatchServices/*
+IncludePaths=Source/Editor/*
+IncludePaths=Source/Runtime/*
+IncludePaths=Source/Developer/*
+ExcludePaths=Config/Localization/*
+ExcludePaths=Source/Developer/NoRedist/CommunityPortalServices/*
+ExcludePaths=Source/Developer/NoRedist/UnrealEngineLauncherServices/*
+ExcludePaths=Source/Developer/NoRedist/BuildPatchServices/*
 InputKeys=Category
 OutputNamespaces=UObjectCategories
 OutputKeys="{MetaDataValue}"
@@ -40,11 +40,11 @@ OutputKeys="{MetaDataValue}"
 ;Gather text from assets
 [GatherTextStep2]
 CommandletClass=GatherTextFromAssets
-IncludePaths=./Content/Editor/*
-IncludePaths=./Content/Editor*
-PackageExtensions=*.umap
-PackageExtensions=*.upk
-PackageExtensions=*.uasset
+IncludePaths=Content/Editor/*
+IncludePaths=Content/Editor*
+PackageFileNameFilters=*.umap
+PackageFileNameFilters=*.upk
+PackageFileNameFilters=*.uasset
 bFixBroken=false
 ShouldGatherBlueprintPinNames=true
 
diff --git a/Engine/Config/Localization/EditorTutorials.ini b/Engine/Config/Localization/EditorTutorials.ini
--- a/Engine/Config/Localization/EditorTutorials.ini
+++ b/Engine/Config/Localization/EditorTutorials.ini
@@ -1,6 +1,6 @@
 [CommonSettings]
-SourcePath=./Content/Localization/EditorTutorials
-DestinationPath=./Content/Localization/EditorTutorials
+SourcePath=Content/Localization/EditorTutorials
+DestinationPath=Content/Localization/EditorTutorials
 ManifestName=EditorTutorials.manifest
 ArchiveName=EditorTutorials.archive
 PortableObjectName=EditorTutorials.po
@@ -16,10 +16,10 @@ ModulesToPreload=IntroTutorials
 ;Gather text from assets
 [GatherTextStep0]
 CommandletClass=GatherTextFromAssets
-IncludePaths=./Content/Tutorial/*
-PackageExtensions=*.umap
-PackageExtensions=*.upk
-PackageExtensions=*.uasset
+IncludePaths=Content/Tutorial/*
+PackageFileNameFilters=*.umap
+PackageFileNameFilters=*.upk
+PackageFileNameFilters=*.uasset
 bFixBroken=false
 
 ;Write Manifest
diff --git a/Engine/Config/Localization/Engine.ini b/Engine/Config/Localization/Engine.ini
--- a/Engine/Config/Localization/Engine.ini
+++ b/Engine/Config/Localization/Engine.ini
@@ -1,6 +1,6 @@
 [CommonSettings]
-SourcePath=./Content/Localization/Engine
-DestinationPath=./Content/Localization/Engine
+SourcePath=Content/Localization/Engine
+DestinationPath=Content/Localization/Engine
 ManifestName=Engine.manifest
 ArchiveName=Engine.archive
 PortableObjectName=Engine.po
@@ -15,36 +15,36 @@ CulturesToGenerate=zh-CN
 ;Gather text from source code
 [GatherTextStep0]
 CommandletClass=GatherTextFromSource
-IncludePaths=./Source/Runtime/
-IncludePaths=./Source/Rocket/
-IncludePaths=./Source/Developer/
-IncludePaths=./Config/
-ExcludePaths=*/Config/Localization/*
-ExcludePaths=*/Source/Developer/NoRedist/CommunityPortalServices/*
-ExcludePaths=*/Source/Developer/NoRedist/UnrealEngineLauncherServices/*
-ExcludePaths=*/Source/Developer/NoRedist/BuildPatchServices/*
-SourceFileSearchFilters=*.cpp
-SourceFileSearchFilters=*.h
-SourceFileSearchFilters=*.c
-SourceFileSearchFilters=*.inl
-SourceFileSearchFilters=*.mm
-SourceFileSearchFilters=*.ini
+SearchDirectoryPaths=Source/Runtime/
+SearchDirectoryPaths=Source/Rocket/
+SearchDirectoryPaths=Source/Developer/
+SearchDirectoryPaths=Config/
+ExcludePaths=Config/Localization/*
+ExcludePaths=Source/Developer/NoRedist/CommunityPortalServices/*
+ExcludePaths=Source/Developer/NoRedist/UnrealEngineLauncherServices/*
+ExcludePaths=Source/Developer/NoRedist/BuildPatchServices/*
+FileNameFilters=*.cpp
+FileNameFilters=*.h
+FileNameFilters=*.c
+FileNameFilters=*.inl
+FileNameFilters=*.mm
+FileNameFilters=*.ini
 
 ;Gather text from assets
 [GatherTextStep1]
 CommandletClass=GatherTextFromAssets
-IncludePaths=./Content/*
-ExcludePaths=*/Content/Editor/*
-ExcludePaths=*/Content/Developers/*
-ExcludePaths=*/Content/Localization/*
-ExcludePaths=*/Content/TestPackages/*
-ExcludePaths=*/Content/QA_Assets/*
-ExcludePaths=*/Content/EngineSounds/*
-ExcludePaths=*/Content/Maps/Automation/*
-ExcludePaths=*/Content/Tutorial/*
-PackageExtensions=*.umap
-PackageExtensions=*.upk
-PackageExtensions=*.uasset
+IncludePaths=Content/*
+ExcludePaths=Content/Editor/*
+ExcludePaths=Content/Developers/*
+ExcludePaths=Content/Localization/*
+ExcludePaths=Content/TestPackages/*
+ExcludePaths=Content/QA_Assets/*
+ExcludePaths=Content/EngineSounds/*
+ExcludePaths=Content/Maps/Automation/*
+ExcludePaths=Content/Tutorial/*
+PackageFileNameFilters=*.umap
+PackageFileNameFilters=*.upk
+PackageFileNameFilters=*.uasset
 bFixBroken=false
 ShouldGatherBlueprintPinNames=true
 
diff --git a/Engine/Config/Localization/Keywords.ini b/Engine/Config/Localization/Keywords.ini
--- a/Engine/Config/Localization/Keywords.ini
+++ b/Engine/Config/Localization/Keywords.ini
@@ -15,11 +15,11 @@ CulturesToGenerate=zh-CN
 ;Gather text from metadata
 [GatherTextStep0]
 CommandletClass=GatherTextFromMetaData
-IncludePaths=./Source/Editor/*
-IncludePaths=./Source/Runtime/*
-IncludePaths=./Source/Developer/*
-ExcludePaths=*/Config/Localization/*
-ExcludePaths=*/Source/Developer/NoRedist/CommunityPortalServices/*
+IncludePaths=Source/Editor/*
+IncludePaths=Source/Runtime/*
+IncludePaths=Source/Developer/*
+ExcludePaths=Config/Localization/*
+ExcludePaths=Source/Developer/NoRedist/CommunityPortalServices/*
 InputKeys=Keywords
 OutputNamespaces=UObjectKeywords
 OutputKeys="{FieldPath}"
diff --git a/Engine/Config/Localization/PortableObjectExport.ini b/Engine/Config/Localization/PortableObjectExport.ini
--- a/Engine/Config/Localization/PortableObjectExport.ini
+++ b/Engine/Config/Localization/PortableObjectExport.ini
@@ -9,40 +9,40 @@ CulturesToGenerate=zh-CN
 
 [GatherTextStep0]
 CommandletClass=InternationalizationExport
-SourcePath=./Content/Localization/Engine
-DestinationPath=./Content/Localization/Engine
+SourcePath=Content/Localization/Engine
+DestinationPath=Content/Localization/Engine
 ManifestName=Engine.manifest
 ArchiveName=Engine.archive
 PortableObjectName=Engine.po
 
 [GatherTextStep1]
 CommandletClass=InternationalizationExport
-SourcePath=./Content/Localization/Editor
-DestinationPath=./Content/Localization/Editor
+SourcePath=Content/Localization/Editor
+DestinationPath=Content/Localization/Editor
 ManifestName=Editor.manifest
 ArchiveName=Editor.archive
 PortableObjectName=Editor.po
 
 [GatherTextStep2]
 CommandletClass=InternationalizationExport
-SourcePath=./Content/Localization/ToolTips
-DestinationPath=./Content/Localization/ToolTips
+SourcePath=Content/Localization/ToolTips
+DestinationPath=Content/Localization/ToolTips
 ManifestName=ToolTips.manifest
 ArchiveName=ToolTips.archive
 PortableObjectName=ToolTips.po
 
 [GatherTextStep3]
 CommandletClass=InternationalizationExport
-SourcePath=./Content/Localization/PropertyNames
-DestinationPath=./Content/Localization/PropertyNames
+SourcePath=Content/Localization/PropertyNames
+DestinationPath=Content/Localization/PropertyNames
 ManifestName=PropertyNames.manifest
 ArchiveName=PropertyNames.archive
 PortableObjectName=PropertyNames.po
 
 [GatherTextStep4]
 CommandletClass=InternationalizationExport
-SourcePath=./Content/Localization/EditorTutorials
-DestinationPath=./Content/Localization/EditorTutorials
+SourcePath=Content/Localization/EditorTutorials
+DestinationPath=Content/Localization/EditorTutorials
 ManifestName=EditorTutorials.manifest
 ArchiveName=EditorTutorials.archive
 PortableObjectName=EditorTutorials.po
\ No newline at end of file
diff --git a/Engine/Config/Localization/PortableObjectImport.ini b/Engine/Config/Localization/PortableObjectImport.ini
--- a/Engine/Config/Localization/PortableObjectImport.ini
+++ b/Engine/Config/Localization/PortableObjectImport.ini
@@ -9,40 +9,40 @@ CulturesToGenerate=zh-CN
 
 [GatherTextStep0]
 CommandletClass=InternationalizationExport
-SourcePath=./Content/Localization/Engine
-DestinationPath=./Content/Localization/Engine
+SourcePath=Content/Localization/Engine
+DestinationPath=Content/Localization/Engine
 ManifestName=Engine.manifest
 ArchiveName=Engine.archive
 PortableObjectName=Engine.po
 
 [GatherTextStep1]
 CommandletClass=InternationalizationExport
-SourcePath=./Content/Localization/Editor
-DestinationPath=./Content/Localization/Editor
+SourcePath=Content/Localization/Editor
+DestinationPath=Content/Localization/Editor
 ManifestName=Editor.manifest
 ArchiveName=Editor.archive
 PortableObjectName=Editor.po
 
 [GatherTextStep2]
 CommandletClass=InternationalizationExport
-SourcePath=./Content/Localization/ToolTips
-DestinationPath=./Content/Localization/ToolTips
+SourcePath=Content/Localization/ToolTips
+DestinationPath=Content/Localization/ToolTips
 ManifestName=ToolTips.manifest
 ArchiveName=ToolTips.archive
 PortableObjectName=ToolTips.po
 
 [GatherTextStep3]
 CommandletClass=InternationalizationExport
-SourcePath=./Content/Localization/PropertyNames
-DestinationPath=./Content/Localization/PropertyNames
+SourcePath=Content/Localization/PropertyNames
+DestinationPath=Content/Localization/PropertyNames
 ManifestName=PropertyNames.manifest
 ArchiveName=PropertyNames.archive
 PortableObjectName=PropertyNames.po
 
 [GatherTextStep4]
 CommandletClass=InternationalizationExport
-SourcePath=./Content/Localization/EditorTutorials
-DestinationPath=./Content/Localization/EditorTutorials
+SourcePath=Content/Localization/EditorTutorials
+DestinationPath=Content/Localization/EditorTutorials
 ManifestName=EditorTutorials.manifest
 ArchiveName=EditorTutorials.archive
 PortableObjectName=EditorTutorials.po
\ No newline at end of file
diff --git a/Engine/Config/Localization/PropertyNames.ini b/Engine/Config/Localization/PropertyNames.ini
--- a/Engine/Config/Localization/PropertyNames.ini
+++ b/Engine/Config/Localization/PropertyNames.ini
@@ -1,6 +1,6 @@
 [CommonSettings]
-SourcePath=./Content/Localization/PropertyNames
-DestinationPath=./Content/Localization/PropertyNames
+SourcePath=Content/Localization/PropertyNames
+DestinationPath=Content/Localization/PropertyNames
 ManifestName=PropertyNames.manifest
 ArchiveName=PropertyNames.archive
 PortableObjectName=PropertyNames.po
@@ -15,11 +15,11 @@ CulturesToGenerate=zh-CN
 ;Gather text from metadata
 [GatherTextStep0]
 CommandletClass=GatherTextFromMetaData
-IncludePaths=./Source/Editor/*
-IncludePaths=./Source/Runtime/*
-IncludePaths=./Source/Developer/*
-ExcludePaths=*/Config/Localization/*
-ExcludePaths=*/Source/Developer/NoRedist/CommunityPortalServices/*
+IncludePaths=Source/Editor/*
+IncludePaths=Source/Runtime/*
+IncludePaths=Source/Developer/*
+ExcludePaths=Config/Localization/*
+ExcludePaths=Source/Developer/NoRedist/CommunityPortalServices/*
 InputKeys=DisplayName
 OutputNamespaces=UObjectDisplayNames
 OutputKeys="{FieldPath}"
diff --git a/Engine/Config/Localization/RegenerateEditor.ini b/Engine/Config/Localization/RegenerateEditor.ini
--- a/Engine/Config/Localization/RegenerateEditor.ini
+++ b/Engine/Config/Localization/RegenerateEditor.ini
@@ -1,5 +1,5 @@
 [RegenerateResources]
-SourcePath=./Content/Localization/Editor
-DestinationPath=./Content/Localization/Editor
+SourcePath=Content/Localization/Editor
+DestinationPath=Content/Localization/Editor
 ManifestName=Editor.manifest
 ResourceName=Editor.locres
\ No newline at end of file
diff --git a/Engine/Config/Localization/RegenerateEditorTutorials.ini b/Engine/Config/Localization/RegenerateEditorTutorials.ini
--- a/Engine/Config/Localization/RegenerateEditorTutorials.ini
+++ b/Engine/Config/Localization/RegenerateEditorTutorials.ini
@@ -1,5 +1,5 @@
 [RegenerateResources]
-SourcePath=./Content/Localization/EditorTutorials
-DestinationPath=./Content/Localization/EditorTutorials
+SourcePath=Content/Localization/EditorTutorials
+DestinationPath=Content/Localization/EditorTutorials
 ManifestName=EditorTutorials.manifest
 ResourceName=EditorTutorials.locres
\ No newline at end of file
diff --git a/Engine/Config/Localization/RegenerateEngine.ini b/Engine/Config/Localization/RegenerateEngine.ini
--- a/Engine/Config/Localization/RegenerateEngine.ini
+++ b/Engine/Config/Localization/RegenerateEngine.ini
@@ -1,5 +1,5 @@
 [RegenerateResources]
-SourcePath=./Content/Localization/Engine
-DestinationPath=./Content/Localization/Engine
+SourcePath=Content/Localization/Engine
+DestinationPath=Content/Localization/Engine
 ManifestName=Engine.manifest
 ResourceName=Engine.locres
\ No newline at end of file
diff --git a/Engine/Config/Localization/RepairData.ini b/Engine/Config/Localization/RepairData.ini
--- a/Engine/Config/Localization/RepairData.ini
+++ b/Engine/Config/Localization/RepairData.ini
@@ -9,74 +9,74 @@ CulturesToGenerate=zh-CN
 ;Repair localization data for Engine
 [GatherTextStep0]
 CommandletClass=RepairLocalizationData
-DestinationPath=./Content/Localization/Engine
+DestinationPath=Content/Localization/Engine
 ManifestName=Engine.manifest
 ArchiveName=Engine.archive
 
 ;Write Localized Text Resource for Engine
 [GatherTextStep1]
 CommandletClass=GenerateTextLocalizationResource
-SourcePath=./Content/Localization/Engine
-DestinationPath=./Content/Localization/Engine
+SourcePath=Content/Localization/Engine
+DestinationPath=Content/Localization/Engine
 ManifestName=Engine.manifest
 ResourceName=Engine.locres
 
 ;Repair localization data for Editor
 [GatherTextStep2]
 CommandletClass=RepairLocalizationData
-DestinationPath=./Content/Localization/Editor
+DestinationPath=Content/Localization/Editor
 ManifestName=Editor.manifest
 ArchiveName=Editor.archive
 
 ;Write Localized Text Resource for Editor
 [GatherTextStep3]
 CommandletClass=GenerateTextLocalizationResource
-SourcePath=./Content/Localization/Editor
-DestinationPath=./Content/Localization/Editor
+SourcePath=Content/Localization/Editor
+DestinationPath=Content/Localization/Editor
 ManifestName=Editor.manifest
 ResourceName=Editor.locres
 
 ;Repair localization data for Editor Tool Tips
 [GatherTextStep4]
 CommandletClass=RepairLocalizationData
-DestinationPath=./Content/Localization/ToolTips
+DestinationPath=Content/Localization/ToolTips
 ManifestName=ToolTips.manifest
 ArchiveName=ToolTips.archive
 
 ;Write Localized Text Resource for Editor Tool Tips
 [GatherTextStep5]
 CommandletClass=GenerateTextLocalizationResource
-SourcePath=./Content/Localization/ToolTips
-DestinationPath=./Content/Localization/ToolTips
+SourcePath=Content/Localization/ToolTips
+DestinationPath=Content/Localization/ToolTips
 ManifestName=ToolTips.manifest
 ResourceName=ToolTips.locres
 
 ;Repair localization data for Editor Property Names
 [GatherTextStep6]
 CommandletClass=RepairLocalizationData
-DestinationPath=./Content/Localization/PropertyNames
+DestinationPath=Content/Localization/PropertyNames
 ManifestName=PropertyNames.manifest
 ArchiveName=PropertyNames.archive
 
 ;Write Localized Text Resource for Editor Property Names
 [GatherTextStep7]
 CommandletClass=GenerateTextLocalizationResource
-SourcePath=./Content/Localization/PropertyNames
-DestinationPath=./Content/Localization/PropertyNames
+SourcePath=Content/Localization/PropertyNames
+DestinationPath=Content/Localization/PropertyNames
 ManifestName=PropertyNames.manifest
 ResourceName=PropertyNames.locres
 
 ;Repair localization data for Editor Tutorials
 [GatherTextStep8]
 CommandletClass=RepairLocalizationData
-DestinationPath=./Content/Localization/PropertyNames
+DestinationPath=Content/Localization/PropertyNames
 ManifestName=PropertyNames.manifest
 ArchiveName=PropertyNames.archive
 
 ;Write Localized Text Resource for Editor Tutorials
 [GatherTextStep9]
 CommandletClass=GenerateTextLocalizationResource
-SourcePath=./Content/Localization/EditorTutorials
-DestinationPath=./Content/Localization/EditorTutorials
+SourcePath=Content/Localization/EditorTutorials
+DestinationPath=Content/Localization/EditorTutorials
 ManifestName=EditorTutorials.manifest
 ResourceName=EditorTutorials.locres
\ No newline at end of file
diff --git a/Engine/Config/Localization/ResourceFileGen.ini b/Engine/Config/Localization/ResourceFileGen.ini
--- a/Engine/Config/Localization/ResourceFileGen.ini
+++ b/Engine/Config/Localization/ResourceFileGen.ini
@@ -9,39 +9,39 @@ CulturesToGenerate=zh-CN
 ;Write Localized Text Resource for Engine
 [GatherTextStep0]
 CommandletClass=GenerateTextLocalizationResource
-SourcePath=./Content/Localization/Engine
-DestinationPath=./Content/Localization/Engine
+SourcePath=Content/Localization/Engine
+DestinationPath=Content/Localization/Engine
 ManifestName=Engine.manifest
 ResourceName=Engine.locres
 
 ;Write Localized Text Resource for Editor
 [GatherTextStep1]
 CommandletClass=GenerateTextLocalizationResource
-SourcePath=./Content/Localization/Editor
-DestinationPath=./Content/Localization/Editor
+SourcePath=Content/Localization/Editor
+DestinationPath=Content/Localization/Editor
 ManifestName=Editor.manifest
 ResourceName=Editor.locres
 
 ;Write Localized Text Resource for Editor Tool Tips
 [GatherTextStep2]
 CommandletClass=GenerateTextLocalizationResource
-SourcePath=./Content/Localization/ToolTips
-DestinationPath=./Content/Localization/ToolTips
+SourcePath=Content/Localization/ToolTips
+DestinationPath=Content/Localization/ToolTips
 ManifestName=ToolTips.manifest
 ResourceName=ToolTips.locres
 
 ;Write Localized Text Resource for Editor Property Names
 [GatherTextStep3]
 CommandletClass=GenerateTextLocalizationResource
-SourcePath=./Content/Localization/PropertyNames
-DestinationPath=./Content/Localization/PropertyNames
+SourcePath=Content/Localization/PropertyNames
+DestinationPath=Content/Localization/PropertyNames
 ManifestName=PropertyNames.manifest
 ResourceName=PropertyNames.locres
 
 ;Write Localized Text Resource for Editor Tutorials
 [GatherTextStep4]
 CommandletClass=GenerateTextLocalizationResource
-SourcePath=./Content/Localization/EditorTutorials
-DestinationPath=./Content/Localization/EditorTutorials
+SourcePath=Content/Localization/EditorTutorials
+DestinationPath=Content/Localization/EditorTutorials
 ManifestName=EditorTutorials.manifest
 ResourceName=EditorTutorials.locres
\ No newline at end of file
diff --git a/Engine/Config/Localization/ToolTips.ini b/Engine/Config/Localization/ToolTips.ini
--- a/Engine/Config/Localization/ToolTips.ini
+++ b/Engine/Config/Localization/ToolTips.ini
@@ -1,6 +1,6 @@
 [CommonSettings]
-SourcePath=./Content/Localization/ToolTips
-DestinationPath=./Content/Localization/ToolTips
+SourcePath=Content/Localization/ToolTips
+DestinationPath=Content/Localization/ToolTips
 ManifestName=ToolTips.manifest
 ArchiveName=ToolTips.archive
 PortableObjectName=ToolTips.po
@@ -15,11 +15,11 @@ CulturesToGenerate=zh-CN
 ;Gather text from metadata
 [GatherTextStep0]
 CommandletClass=GatherTextFromMetaData
-IncludePaths=./Source/Editor/*
-IncludePaths=./Source/Runtime/*
-IncludePaths=./Source/Developer/*
-ExcludePaths=*/Config/Localization/*
-ExcludePaths=*/Source/Developer/NoRedist/CommunityPortalServices/*
+IncludePaths=Source/Editor/*
+IncludePaths=Source/Runtime/*
+IncludePaths=Source/Developer/*
+ExcludePaths=Config/Localization/*
+ExcludePaths=Source/Developer/NoRedist/CommunityPortalServices/*
 InputKeys=ToolTip
 OutputNamespaces=UObjectToolTips
 OutputKeys="{FieldPath}"
diff --git a/Engine/Config/Localization/WordCount.ini b/Engine/Config/Localization/WordCount.ini
--- a/Engine/Config/Localization/WordCount.ini
+++ b/Engine/Config/Localization/WordCount.ini
@@ -10,34 +10,34 @@ CulturesToGenerate=zh-CN
 ;Write Text Localization Report for Engine
 [GatherTextStep0]
 CommandletClass=GenerateTextLocalizationReport
-SourcePath=./Content/Localization/Engine
+SourcePath=Content/Localization/Engine
 ManifestName=Engine.manifest
 WordCountReportName=Engine.csv
 
 ;Write Text Localization Report for Editor
 [GatherTextStep1]
 CommandletClass=GenerateTextLocalizationReport
-SourcePath=./Content/Localization/Editor
+SourcePath=Content/Localization/Editor
 ManifestName=Editor.manifest
 WordCountReportName=Editor.csv
 
 ;Write Text Localization Report for Editor Tool Tips
 [GatherTextStep2]
 CommandletClass=GenerateTextLocalizationReport
-SourcePath=./Content/Localization/ToolTips
+SourcePath=Content/Localization/ToolTips
 ManifestName=ToolTips.manifest
 WordCountReportName=ToolTips.csv
 
 ;Write Text Localization Report for Editor Property Names
 [GatherTextStep3]
 CommandletClass=GenerateTextLocalizationReport
-SourcePath=./Content/Localization/PropertyNames
+SourcePath=Content/Localization/PropertyNames
 ManifestName=PropertyNames.manifest
 WordCountReportName=PropertyNames.csv
 
 ;Write Text Localization Report for Editor Tutorials
 [GatherTextStep4]
 CommandletClass=GenerateTextLocalizationReport
-SourcePath=./Content/Localization/EditorTutorials
+SourcePath=Content/Localization/EditorTutorials
 ManifestName=EditorTutorials.manifest
 WordCountReportName=EditorTutorials.csv
\ No newline at end of file
diff --git a/Engine/Source/Editor/EditorStyle/Private/SlateEditorStyle.cpp b/Engine/Source/Editor/EditorStyle/Private/SlateEditorStyle.cpp
--- a/Engine/Source/Editor/EditorStyle/Private/SlateEditorStyle.cpp
+++ b/Engine/Source/Editor/EditorStyle/Private/SlateEditorStyle.cpp
@@ -6518,12 +6518,19 @@ void FSlateEditorStyle::FStyle::SetupLocalizationDashboardStyles()
 	Set("LocalizationTargetEditor.RefreshWordCounts", new IMAGE_BRUSH("Icons/Icon_Localisation_Refresh_Word_Counts_40x", Icon40x40));
 	Set("LocalizationTargetEditor.Compile", new IMAGE_BRUSH( "Icons/Icon_Localisation_Compile_Translations_40x", Icon40x40));
 
+	Set("LocalizationTargetEditor.DirectoryPicker", new IMAGE_BRUSH( "Icons/ellipsis_12x", Icon12x12 ));
+	Set("LocalizationTargetEditor.GatherSettingsIcon_Valid", new IMAGE_BRUSH("Settings/Settings_Good", Icon16x16));
+	Set("LocalizationTargetEditor.GatherSettingsIcon_Warning", new IMAGE_BRUSH("Settings/Settings_Warning", Icon16x16));
+
 	Set("LocalizationTargetEditor.NativeCulture", new IMAGE_BRUSH( "Icons/Star_16x", Icon16x16 ) );
+
 	Set("LocalizationTargetEditor.EditTranslations", new IMAGE_BRUSH("Icons/icon_file_open_16px", Icon16x16));
 	Set("LocalizationTargetEditor.ImportForCulture", new IMAGE_BRUSH("Icons/Icon_Localisation_Import_All_16x", Icon16x16));
 	Set("LocalizationTargetEditor.ExportForCulture", new IMAGE_BRUSH("Icons/Icon_Localisation_Export_All_16x", Icon16x16));
 	Set("LocalizationTargetEditor.DeleteCulture", new IMAGE_BRUSH("Icons/Cross_12x", Icon12x12 ) );
 
+	Set("LocalizationTargetEditor.GatherSettings.AddMetaDataTextKeyPatternArgument", new IMAGE_BRUSH("Icons/icon_Blueprint_AddVariable_40px", Icon16x16 ) );
+
 	Set( "LocalizationDashboard.CommandletLog.Text", FTextBlockStyle(NormalText)
 		.SetFont( TTF_CORE_FONT( "Fonts/Roboto-Regular", 8 ) )
 		.SetShadowOffset( FVector2D::ZeroVector )
diff --git a/Engine/Source/Editor/LocalizationDashboard/Private/GatherTextDetailCustomizations.cpp b/Engine/Source/Editor/LocalizationDashboard/Private/GatherTextDetailCustomizations.cpp
new file mode 100644
--- /dev/null
+++ b/Engine/Source/Editor/LocalizationDashboard/Private/GatherTextDetailCustomizations.cpp
@@ -0,0 +1,975 @@
+// Copyright 1998-2015 Epic Games, Inc. All Rights Reserved.
+
+#include "LocalizationDashboardPrivatePCH.h"
+#include "GatherTextDetailCustomizations.h"
+#include "LocalizationTargetTypes.h"
+#include "DetailWidgetRow.h"
+#include "PropertyHandle.h"
+#include "IDetailPropertyRow.h"
+#include "SEditableTextBox.h"
+#include "IPropertyUtilities.h"
+#include "IPropertyTypeCustomization.h"
+
+#define LOCTEXT_NAMESPACE "GatherTextDetailCustomizations"
+
+namespace
+{
+	class SGatherTextConfigurationErrorHint : public SCompoundWidget
+	{
+		SLATE_BEGIN_ARGS(SGatherTextConfigurationErrorHint) {}
+			SLATE_ATTRIBUTE(FText, ErrorText)
+		SLATE_END_ARGS()
+
+	public:
+		void Construct(const FArguments& Arguments);
+		virtual void Tick(const FGeometry& AllottedGeometry, const double InCurrentTime, const float InDeltaTime) override;
+
+	private:
+		TAttribute<FText> ErrorText;
+		TSharedPtr<SErrorHint> ErrorHint;
+	};
+
+	void SGatherTextConfigurationErrorHint::Construct(const FArguments& Arguments)
+	{
+		ErrorText = Arguments._ErrorText;
+
+		ChildSlot
+			[
+				SAssignNew(ErrorHint, SErrorHint)
+			];
+	}
+
+	void SGatherTextConfigurationErrorHint::Tick(const FGeometry& AllottedGeometry, const double InCurrentTime, const float InDeltaTime)
+	{
+		SCompoundWidget::Tick(AllottedGeometry, InCurrentTime, InDeltaTime);
+
+		if (ErrorText.IsSet() && ErrorHint.IsValid())
+		{
+			ErrorHint->SetError(ErrorText.Get());
+		}
+	}
+}
+
+namespace
+{
+	class SGatherTextPathPicker : public SCompoundWidget
+	{
+		SLATE_BEGIN_ARGS(SGatherTextPathPicker)
+			: _ShouldCoercePathAsWildcardPattern(false)
+		{}
+			SLATE_ARGUMENT(bool, ShouldCoercePathAsWildcardPattern)
+		SLATE_END_ARGS()
+
+	public:
+		void Construct(const FArguments& Args, const TSharedRef<IPropertyHandle>& PathStringPropertyHandle);
+
+	private:
+		FReply PathPickerOnClicked();
+
+	private:
+		TSharedPtr<IPropertyHandle> PathStringPropertyHandle;
+		bool ShouldCoercePathAsWildcardPattern;
+	};
+
+	void SGatherTextPathPicker::Construct(const FArguments& Args, const TSharedRef<IPropertyHandle>& InPathStringPropertyHandle)
+	{
+		PathStringPropertyHandle = InPathStringPropertyHandle;
+		ShouldCoercePathAsWildcardPattern = Args._ShouldCoercePathAsWildcardPattern;
+
+		TArray<UObject*> OuterObjects;
+		PathStringPropertyHandle->GetOuterObjects(OuterObjects);
+		ULocalizationTarget* const LocalizationTarget = Cast<ULocalizationTarget>(OuterObjects.Top());
+			
+		ChildSlot
+			[
+				SNew(SHorizontalBox)
+				+SHorizontalBox::Slot()
+				.AutoWidth()
+				.VAlign(VAlign_Center)
+				.Padding(2.0f, 0.0f)
+				[
+					SNew(STextBlock)
+					.Text( FText::FromString( FString::Printf( TEXT("%s/"), *(FPaths::GetBaseFilename(LocalizationTarget->IsMemberOfEngineTargetSet() ? FPaths::EngineDir() : FPaths::GameDir())) ) ) )
+				]
+				+SHorizontalBox::Slot()
+				.AutoWidth()
+				.VAlign(VAlign_Center)
+				[
+					SNew(SBox)
+					.MinDesiredWidth(125.0f)
+					[
+						PathStringPropertyHandle.IsValid() && PathStringPropertyHandle->IsValidHandle() ? PathStringPropertyHandle->CreatePropertyValueWidget() : SNullWidget::NullWidget
+					]
+				]
+				+SHorizontalBox::Slot()
+				.AutoWidth()
+				.Padding(FMargin(4.0f, 0.0f, 0.0f, 0.0f))
+				.VAlign(VAlign_Center)
+				[
+					SNew(SButton)
+					.ButtonStyle( FEditorStyle::Get(), "HoverHintOnly" )
+					.ToolTipText( LOCTEXT( "PatchPickerToolTipText", "Choose a directory.") )
+					.OnClicked(this, &SGatherTextPathPicker::PathPickerOnClicked)
+					.ContentPadding( 2.0f )
+					.ForegroundColor( FSlateColor::UseForeground() )
+					.IsFocusable( false )
+					[
+						SNew( SImage )
+						.Image( FEditorStyle::GetBrush("LocalizationTargetEditor.DirectoryPicker") )
+						.ColorAndOpacity( FSlateColor::UseForeground() )
+					]
+				]
+			];
+	}
+
+	FReply SGatherTextPathPicker::PathPickerOnClicked()
+	{
+		if (!PathStringPropertyHandle.IsValid() || !PathStringPropertyHandle->IsValidHandle())
+		{
+			return FReply::Handled();
+		}
+
+		IDesktopPlatform* const DesktopPlatform = FDesktopPlatformModule::Get();
+		check(DesktopPlatform);
+
+		TArray<UObject*> OuterObjects;
+		PathStringPropertyHandle->GetOuterObjects(OuterObjects);
+		ULocalizationTarget* const LocalizationTarget = CastChecked<ULocalizationTarget>(OuterObjects.Top());
+		const FString DesiredRootPath = FPaths::ConvertRelativePathToFull(LocalizationTarget->IsMemberOfEngineTargetSet() ? FPaths::EngineDir() : FPaths::GameDir());
+
+		const TSharedPtr<SWindow> ParentWindow = FSlateApplication::Get().FindWidgetWindow(AsShared());
+		const TSharedPtr<FGenericWindow> ParentGenericWindow = ParentWindow.IsValid() ? ParentWindow->GetNativeWindow() : nullptr;
+		const void* ParentWindowHandle = ParentGenericWindow.IsValid() ? ParentGenericWindow->GetOSWindowHandle() : nullptr;
+		const FText DialogTitle = LOCTEXT("SelectSearchDirectoryDialogTitle", "Select Directory Containing Text Files");
+
+		FString CurrentPath;
+		PathStringPropertyHandle->GetValue(CurrentPath);
+		if (FPaths::IsRelative(CurrentPath))
+		{
+			CurrentPath = FPaths::Combine(*DesiredRootPath, *CurrentPath);
+			int32 WildcardStarIndex;
+			CurrentPath.FindLastChar(TEXT('*'), WildcardStarIndex);
+			if (WildcardStarIndex != INDEX_NONE)
+			{
+				CurrentPath = CurrentPath.Left(WildcardStarIndex);
+			}
+		}
+
+		const FString DefaultPath = CurrentPath.IsEmpty() ? DesiredRootPath : CurrentPath;
+
+		FString NewPath = DefaultPath;
+		if (DesktopPlatform->OpenDirectoryDialog(ParentWindowHandle, DialogTitle.ToString(), DefaultPath, NewPath))
+		{
+			FPaths::MakePathRelativeTo(NewPath, *DesiredRootPath);
+
+			if (ShouldCoercePathAsWildcardPattern)
+			{
+				if (!NewPath.Contains(TEXT("*")))
+				{
+					NewPath = FPaths::Combine(*NewPath, TEXT("*"));
+				}
+			}
+			PathStringPropertyHandle->SetValue(NewPath);
+		}
+
+		return FReply::Handled();
+	}
+}
+
+void FGatherTextSearchDirectoryStructCustomization::CustomizeStructHeader( TSharedRef<IPropertyHandle> StructPropertyHandle, class FDetailWidgetRow& HeaderRow, IStructCustomizationUtils& StructCustomizationUtils )
+{
+	const TSharedPtr<IPropertyHandle> PathPropertyHandle = StructPropertyHandle->GetChildHandle(GET_MEMBER_NAME_CHECKED(FGatherTextSearchDirectory, Path));
+
+	const auto& GetErrorText = [StructPropertyHandle]() -> FText
+	{
+		if (StructPropertyHandle->IsValidHandle())
+		{
+			TArray<const void*> RawData;
+			StructPropertyHandle->AccessRawData(RawData);
+			const FGatherTextSearchDirectory& Struct = *(reinterpret_cast<const FGatherTextSearchDirectory*>(RawData.Top()));
+
+			TArray<UObject*> OuterObjects;
+			StructPropertyHandle->GetOuterObjects(OuterObjects);
+			ULocalizationTarget* const LocalizationTarget = Cast<ULocalizationTarget>(OuterObjects.Top());
+			if (LocalizationTarget)
+			{
+				const FString DesiredRootPath = FPaths::ConvertRelativePathToFull(LocalizationTarget->IsMemberOfEngineTargetSet() ? FPaths::EngineDir() : FPaths::GameDir());
+				FText Error;
+				if (!Struct.Validate(DesiredRootPath, Error))
+				{
+					return Error;
+				}
+			}
+		}
+
+		return FText::GetEmpty();
+	};
+
+	HeaderRow
+		.NameContent()
+		[
+			SNew(SHorizontalBox)
+			+SHorizontalBox::Slot()
+			.AutoWidth()
+			[
+				StructPropertyHandle->CreatePropertyNameWidget()
+			]
+			+SHorizontalBox::Slot()
+			.AutoWidth()
+			[
+				SNew(SGatherTextConfigurationErrorHint)
+				.ErrorText_Lambda(GetErrorText)
+			]
+		]
+		.ValueContent()
+		.MaxDesiredWidth(TOptional<float>())
+		[
+			SNew(SGatherTextPathPicker, PathPropertyHandle.ToSharedRef())
+		];
+}
+
+void FGatherTextIncludePathStructCustomization::CustomizeStructHeader( TSharedRef<IPropertyHandle> StructPropertyHandle, class FDetailWidgetRow& HeaderRow, IStructCustomizationUtils& StructCustomizationUtils )
+{
+	const TSharedPtr<IPropertyHandle> PatternPropertyHandle = StructPropertyHandle->GetChildHandle(GET_MEMBER_NAME_CHECKED(FGatherTextIncludePath, Pattern));
+
+	const auto& GetErrorText = [StructPropertyHandle]() -> FText
+	{
+		if (StructPropertyHandle->IsValidHandle())
+		{
+			TArray<const void*> RawData;
+			StructPropertyHandle->AccessRawData(RawData);
+			const FGatherTextIncludePath& Struct = *(reinterpret_cast<const FGatherTextIncludePath*>(RawData.Top()));
+
+			TArray<UObject*> OuterObjects;
+			StructPropertyHandle->GetOuterObjects(OuterObjects);
+			ULocalizationTarget* const LocalizationTarget = Cast<ULocalizationTarget>(OuterObjects.Top());
+			if (LocalizationTarget)
+			{
+				const FString DesiredRootPath = FPaths::ConvertRelativePathToFull(LocalizationTarget->IsMemberOfEngineTargetSet() ? FPaths::EngineDir() : FPaths::GameDir());
+				FText Error;
+				if (!Struct.Validate(DesiredRootPath, Error))
+				{
+					return Error;
+				}
+			}
+		}
+
+		return FText::GetEmpty();
+	};
+
+	HeaderRow
+		.NameContent()
+		[
+			SNew(SHorizontalBox)
+			+SHorizontalBox::Slot()
+			.AutoWidth()
+			[
+				StructPropertyHandle->CreatePropertyNameWidget()
+			]
+			+SHorizontalBox::Slot()
+			.AutoWidth()
+			[
+				SNew(SGatherTextConfigurationErrorHint)
+				.ErrorText_Lambda(GetErrorText)
+			]
+		]
+		.ValueContent()
+		.MaxDesiredWidth(TOptional<float>())
+		[
+			SNew(SGatherTextPathPicker, PatternPropertyHandle.ToSharedRef())
+			.ShouldCoercePathAsWildcardPattern(true)
+		];
+}
+
+void FGatherTextExcludePathStructCustomization::CustomizeStructHeader( TSharedRef<IPropertyHandle> StructPropertyHandle, class FDetailWidgetRow& HeaderRow, IStructCustomizationUtils& StructCustomizationUtils )
+{
+	const TSharedPtr<IPropertyHandle> PatternPropertyHandle = StructPropertyHandle->GetChildHandle(GET_MEMBER_NAME_CHECKED(FGatherTextExcludePath, Pattern));
+
+	const auto& GetErrorText = [StructPropertyHandle]() -> FText
+	{
+		if (StructPropertyHandle->IsValidHandle())
+		{
+			TArray<const void*> RawData;
+			StructPropertyHandle->AccessRawData(RawData);
+			const FGatherTextExcludePath& Struct = *(reinterpret_cast<const FGatherTextExcludePath*>(RawData.Top()));
+
+			FText Error;
+			if (!Struct.Validate(Error))
+			{
+				return Error;
+			}
+		}
+
+		return FText::GetEmpty();
+	};
+
+	HeaderRow
+		.NameContent()
+		[
+			SNew(SHorizontalBox)
+			+SHorizontalBox::Slot()
+			.AutoWidth()
+			[
+				StructPropertyHandle->CreatePropertyNameWidget()
+			]
+			+SHorizontalBox::Slot()
+			.AutoWidth()
+			[
+				SNew(SGatherTextConfigurationErrorHint)
+				.ErrorText_Lambda(GetErrorText)
+			]
+		]
+		.ValueContent()
+		.MaxDesiredWidth(TOptional<float>())
+		[
+			SNew(SGatherTextPathPicker, PatternPropertyHandle.ToSharedRef())
+			.ShouldCoercePathAsWildcardPattern(true)
+		];
+}
+
+void FGatherTextFileExtensionStructCustomization::CustomizeStructHeader( TSharedRef<IPropertyHandle> StructPropertyHandle, class FDetailWidgetRow& HeaderRow, IStructCustomizationUtils& StructCustomizationUtils )
+{
+	const TSharedPtr<IPropertyHandle> PatternPropertyHandle = StructPropertyHandle->GetChildHandle(GET_MEMBER_NAME_CHECKED(FGatherTextFileExtension, Pattern));
+
+	const auto& GetErrorText = [StructPropertyHandle]() -> FText
+	{
+		if (StructPropertyHandle->IsValidHandle())
+		{
+			TArray<const void*> RawData;
+			StructPropertyHandle->AccessRawData(RawData);
+			const FGatherTextFileExtension& Struct = *(reinterpret_cast<const FGatherTextFileExtension*>(RawData.Top()));
+
+			FText Error;
+			if (!Struct.Validate(Error))
+			{
+				return Error;
+			}
+		}
+
+		return FText::GetEmpty();
+	};
+
+	HeaderRow
+		.NameContent()
+		[
+			SNew(SHorizontalBox)
+			+SHorizontalBox::Slot()
+			.AutoWidth()
+			[
+				StructPropertyHandle->CreatePropertyNameWidget()
+			]
+			+SHorizontalBox::Slot()
+			.AutoWidth()
+			[
+				SNew(SGatherTextConfigurationErrorHint)
+				.ErrorText_Lambda(GetErrorText)
+			]
+		]
+		.ValueContent()
+		[
+			SNew(SHorizontalBox)
+			+SHorizontalBox::Slot()
+			.AutoWidth()
+			.VAlign(VAlign_Center)
+			.Padding(2.0f, 0.0f)
+			[
+				SNew(STextBlock)
+				.Text(FText::FromString(TEXT("*.")))
+			]
+			+SHorizontalBox::Slot()
+			.FillWidth(1.0f)
+			.VAlign(VAlign_Center)
+			[
+				PatternPropertyHandle.IsValid() && PatternPropertyHandle->IsValidHandle() ? PatternPropertyHandle->CreatePropertyValueWidget() : SNullWidget::NullWidget
+			]
+		];
+}
+
+namespace
+{
+	class SConfigurationValidity : public SCompoundWidget
+	{
+		SLATE_BEGIN_ARGS(SConfigurationValidity) {}
+			SLATE_ATTRIBUTE(FText, ConfigurationError)
+		SLATE_END_ARGS()
+
+	public:
+		void Construct(const FArguments& Arguments);
+
+	private:
+		const FSlateBrush* GetImageBrush() const;
+
+	private:
+		TAttribute<FText> ConfigurationError;
+	};
+
+	void SConfigurationValidity::Construct(const FArguments& Arguments)
+	{
+		ConfigurationError = Arguments._ConfigurationError;
+
+		const auto& GetToolTipText = [this]() -> FText
+		{
+			FText ToolTipText = ConfigurationError.Get(FText::GetEmpty());
+			if (ToolTipText.IsEmpty())
+			{
+				ToolTipText = LOCTEXT("ValidGatherConfigurationToolTip", "Configuration settings are valid.");
+			}
+			return ToolTipText;
+		};
+
+		ChildSlot
+			.HAlign(HAlign_Center)
+			.VAlign(VAlign_Center)
+			[
+				SNew(SImage)
+				.Image(this, &SConfigurationValidity::GetImageBrush)
+				.ToolTipText_Lambda(GetToolTipText)
+			];
+	}
+
+	const FSlateBrush* SConfigurationValidity::GetImageBrush() const
+	{
+		return ConfigurationError.Get(FText::GetEmpty()).IsEmpty() ? FEditorStyle::GetBrush("LocalizationTargetEditor.GatherSettingsIcon_Valid") : FEditorStyle::GetBrush("LocalizationTargetEditor.GatherSettingsIcon_Warning");
+	}
+}
+
+void FGatherTextFromTextFilesConfigurationStructCustomization::CustomizeStructHeader( TSharedRef<IPropertyHandle> StructPropertyHandle, class FDetailWidgetRow& HeaderRow, IStructCustomizationUtils& StructCustomizationUtils )
+{
+	const TSharedPtr<IPropertyHandle> IsEnabledPropertyHandle = StructPropertyHandle->GetChildHandle(GET_MEMBER_NAME_CHECKED(FGatherTextFromTextFilesConfiguration, IsEnabled));
+	if (IsEnabledPropertyHandle.IsValid() && IsEnabledPropertyHandle->IsValidHandle())
+	{
+		IsEnabledPropertyHandle->MarkHiddenByCustomization();
+	}
+
+	const auto& GetConfigurationError = [StructPropertyHandle]() -> FText
+	{
+		TArray<const void*> RawData;
+		StructPropertyHandle->AccessRawData(RawData);
+		const FGatherTextFromTextFilesConfiguration& Struct = *(reinterpret_cast<const FGatherTextFromTextFilesConfiguration*>(RawData.Top()));
+		
+		TArray<UObject*> OuterObjects;
+		StructPropertyHandle->GetOuterObjects(OuterObjects);
+		ULocalizationTarget* const LocalizationTarget = Cast<ULocalizationTarget>(OuterObjects.Top());
+		if (LocalizationTarget)
+		{
+			const FString DesiredRootPath = FPaths::ConvertRelativePathToFull(LocalizationTarget->IsMemberOfEngineTargetSet() ? FPaths::EngineDir() : FPaths::GameDir());
+			FText Error;
+			if (!Struct.Validate(DesiredRootPath, Error))
+			{
+				return Error;
+			}
+		}
+
+		return FText::GetEmpty();
+	};
+
+	const auto& GetValidityWidgetVisibility = [IsEnabledPropertyHandle]() -> EVisibility
+	{
+		bool Value = false;
+		if (IsEnabledPropertyHandle.IsValid() && IsEnabledPropertyHandle->IsValidHandle())
+		{
+			IsEnabledPropertyHandle->GetValue(Value);
+		}
+		return Value ? EVisibility::Visible : EVisibility::Hidden;
+	};
+		
+	HeaderRow
+		.NameContent()
+		[
+			SNew(SHorizontalBox)
+			+SHorizontalBox::Slot()
+			.AutoWidth()
+			.VAlign(VAlign_Center)
+			.Padding(2.0f, 0.0f)
+			[
+				IsEnabledPropertyHandle.IsValid() && IsEnabledPropertyHandle->IsValidHandle() ? IsEnabledPropertyHandle->CreatePropertyValueWidget() : SNullWidget::NullWidget
+			]
+			+SHorizontalBox::Slot()
+			.AutoWidth()
+			.VAlign(VAlign_Center)
+			[
+				StructPropertyHandle->CreatePropertyNameWidget()
+			]
+		]
+		.ValueContent()
+		[
+			SNew(SHorizontalBox)
+			+SHorizontalBox::Slot()
+			.AutoWidth()
+			[
+				StructPropertyHandle->CreatePropertyValueWidget()
+			]
+			+SHorizontalBox::Slot()
+			.AutoWidth()
+			[
+				SNew(SConfigurationValidity)
+				.Visibility_Lambda(GetValidityWidgetVisibility)
+				.ConfigurationError_Lambda(GetConfigurationError)
+			]
+		];
+}
+
+void FGatherTextFromTextFilesConfigurationStructCustomization::CustomizeStructChildren( TSharedRef<IPropertyHandle> StructPropertyHandle, class IDetailChildrenBuilder& ChildBuilder, IStructCustomizationUtils& StructCustomizationUtils )
+{
+	uint32 ChildCount;
+	StructPropertyHandle->GetNumChildren(ChildCount);
+
+	for (uint32 i = 0; i < ChildCount; ++i)
+	{
+		const TSharedPtr<IPropertyHandle> ChildPropertyHandle = StructPropertyHandle->GetChildHandle(i);
+		if (!ChildPropertyHandle->IsCustomized())
+		{
+			ChildBuilder.AddChildProperty(ChildPropertyHandle.ToSharedRef());
+		}
+	}
+}
+
+void FGatherTextFromPackagesConfigurationStructCustomization::CustomizeStructHeader( TSharedRef<IPropertyHandle> StructPropertyHandle, class FDetailWidgetRow& HeaderRow, IStructCustomizationUtils& StructCustomizationUtils )
+{
+	const TSharedPtr<IPropertyHandle> IsEnabledPropertyHandle = StructPropertyHandle->GetChildHandle(GET_MEMBER_NAME_CHECKED(FGatherTextFromPackagesConfiguration, IsEnabled));
+	if (IsEnabledPropertyHandle.IsValid() && IsEnabledPropertyHandle->IsValidHandle())
+	{
+		IsEnabledPropertyHandle->MarkHiddenByCustomization();
+	}
+
+	const auto& GetConfigurationError = [StructPropertyHandle]() -> FText
+	{
+		TArray<const void*> RawData;
+		StructPropertyHandle->AccessRawData(RawData);
+		const FGatherTextFromPackagesConfiguration& Struct = *(reinterpret_cast<const FGatherTextFromPackagesConfiguration*>(RawData.Top()));
+
+		TArray<UObject*> OuterObjects;
+		StructPropertyHandle->GetOuterObjects(OuterObjects);
+		ULocalizationTarget* const LocalizationTarget = Cast<ULocalizationTarget>(OuterObjects.Top());
+		if (LocalizationTarget)
+		{
+			const FString DesiredRootPath = FPaths::ConvertRelativePathToFull(LocalizationTarget->IsMemberOfEngineTargetSet() ? FPaths::EngineDir() : FPaths::GameDir());
+			FText Error;
+			if (!Struct.Validate(DesiredRootPath, Error))
+			{
+				return Error;
+			}
+		}
+
+		return FText::GetEmpty();
+	};
+
+	const auto& GetValidityWidgetVisibility = [IsEnabledPropertyHandle]() -> EVisibility
+	{
+		bool Value = false;
+		if (IsEnabledPropertyHandle.IsValid() && IsEnabledPropertyHandle->IsValidHandle())
+		{
+			IsEnabledPropertyHandle->GetValue(Value);
+		}
+		return Value ? EVisibility::Visible : EVisibility::Hidden;
+	};
+		
+	HeaderRow
+		.NameContent()
+		[
+			SNew(SHorizontalBox)
+			+SHorizontalBox::Slot()
+			.AutoWidth()
+			.VAlign(VAlign_Center)
+			.Padding(2.0f, 0.0f)
+			[
+				IsEnabledPropertyHandle.IsValid() && IsEnabledPropertyHandle->IsValidHandle() ? IsEnabledPropertyHandle->CreatePropertyValueWidget() : SNullWidget::NullWidget
+			]
+			+SHorizontalBox::Slot()
+			.AutoWidth()
+			.VAlign(VAlign_Center)
+			[
+				StructPropertyHandle->CreatePropertyNameWidget()
+			]
+		]
+		.ValueContent()
+		[
+			SNew(SHorizontalBox)
+			+SHorizontalBox::Slot()
+			.AutoWidth()
+			[
+				StructPropertyHandle->CreatePropertyValueWidget()
+			]
+			+SHorizontalBox::Slot()
+			.AutoWidth()
+			[
+				SNew(SConfigurationValidity)
+				.Visibility_Lambda(GetValidityWidgetVisibility)
+				.ConfigurationError_Lambda(GetConfigurationError)
+			]
+		];
+}
+
+void FGatherTextFromPackagesConfigurationStructCustomization::CustomizeStructChildren( TSharedRef<IPropertyHandle> StructPropertyHandle, class IDetailChildrenBuilder& ChildBuilder, IStructCustomizationUtils& StructCustomizationUtils )
+{
+	uint32 ChildCount;
+	StructPropertyHandle->GetNumChildren(ChildCount);
+
+	for (uint32 i = 0; i < ChildCount; ++i)
+	{
+		const TSharedPtr<IPropertyHandle> ChildPropertyHandle = StructPropertyHandle->GetChildHandle(i);
+		if (!ChildPropertyHandle->IsCustomized())
+		{
+			ChildBuilder.AddChildProperty(ChildPropertyHandle.ToSharedRef());
+		}
+	}
+}
+
+namespace
+{
+	class SMetaDataTextKeyPatternWidget : public SCompoundWidget
+	{
+		SLATE_BEGIN_ARGS(SMetaDataTextKeyPatternWidget) {}
+		SLATE_END_ARGS()
+
+	public:
+		void Construct(const FArguments& Args, const TSharedRef<IPropertyHandle>& PatternPropertyHandle);
+
+	private:
+		FText GetText() const;
+		void OnTextCommitted(const FText& NewText, ETextCommit::Type CommitInfo);
+		void OnPlaceHolderSelected(TSharedPtr<FString> PlaceHolderString, ESelectInfo::Type SelectInfo);
+
+	private:
+		TSharedPtr<IPropertyHandle> PatternPropertyHandle;
+		TSharedPtr<SEditableTextBox> EditableTextBox;
+		TSharedPtr<SComboButton> PlaceHolderComboButton;
+		TArray< TSharedPtr<FString> > PossiblePlaceHoldersListItemsSource;
+	};
+
+	void SMetaDataTextKeyPatternWidget::Construct(const FArguments& Args, const TSharedRef<IPropertyHandle>& InPatternPropertyHandle)
+	{
+		PatternPropertyHandle = InPatternPropertyHandle;
+
+		for(const FString& PossiblePlaceHolder : FMetaDataTextKeyPattern::PossiblePlaceHolders)
+		{
+			PossiblePlaceHoldersListItemsSource.Add( MakeShareable(new FString(PossiblePlaceHolder)) );
+		}
+
+		const auto& OnGenerateRow = [](TSharedPtr<FString> PlaceHolderString, const TSharedRef<STableViewBase>& Table) -> TSharedRef<ITableRow>
+		{
+			return SNew(STableRow< TSharedPtr<FString> >, Table)
+				.Content()
+				[
+					SNew(SBox)
+					.Padding(2.0f)
+					[
+						SNew(STextBlock)
+						.Text(FText::FromString(*PlaceHolderString.Get()))
+					]
+				];
+		};
+
+		ChildSlot
+			[
+				SNew(SHorizontalBox)
+				+SHorizontalBox::Slot()
+				.AutoWidth()
+				.VAlign(VAlign_Center)
+				[
+					SNew(SBox)
+					.MinDesiredWidth(125.0f)
+					[
+						SAssignNew(EditableTextBox, SEditableTextBox)
+						.Text(this, &SMetaDataTextKeyPatternWidget::GetText)
+						.OnTextCommitted(this, &SMetaDataTextKeyPatternWidget::OnTextCommitted)
+					]
+				]
+				+SHorizontalBox::Slot()
+				.AutoWidth()
+				.VAlign(VAlign_Center)
+				.Padding(2.0f, 0.0f)
+				[
+					SAssignNew(PlaceHolderComboButton, SComboButton)
+					.ButtonStyle(FEditorStyle::Get(), "HoverHintOnly")
+					.ButtonContent()
+					[
+						SNew(SImage)
+						.Image(FEditorStyle::GetBrush("LocalizationTargetEditor.GatherSettings.AddMetaDataTextKeyPatternArgument"))
+					]
+					.OnGetMenuContent_Lambda([&]() -> TSharedRef<SWidget>
+					{
+						return SNew(SListView< TSharedPtr<FString> >)
+							.ListItemsSource(&PossiblePlaceHoldersListItemsSource)
+							.OnGenerateRow_Lambda(OnGenerateRow)
+							.OnSelectionChanged(this, &SMetaDataTextKeyPatternWidget::OnPlaceHolderSelected);
+					})
+				]
+			];
+	}
+
+	FText SMetaDataTextKeyPatternWidget::GetText() const
+	{
+		if (!PatternPropertyHandle.IsValid() || !PatternPropertyHandle->IsValidHandle())
+		{
+			return FText::GetEmpty();
+		}
+
+		FString PatternString;
+		PatternPropertyHandle->GetValue(PatternString);
+
+		return FText::FromString(PatternString);
+	}
+
+	void SMetaDataTextKeyPatternWidget::OnTextCommitted( const FText& NewText, ETextCommit::Type CommitInfo )
+	{
+		if (!PatternPropertyHandle.IsValid() || !PatternPropertyHandle->IsValidHandle())
+		{
+			return;
+		}
+
+		const FString PatternString = NewText.ToString();
+		PatternPropertyHandle->SetValue(PatternString);
+	}
+
+	void SMetaDataTextKeyPatternWidget::OnPlaceHolderSelected(TSharedPtr<FString> PlaceHolderString, ESelectInfo::Type SelectInfo)
+	{
+		if (PlaceHolderString.IsValid() && EditableTextBox.IsValid())
+		{
+			EditableTextBox->SetText(FText::FromString(EditableTextBox->GetText().ToString() + *PlaceHolderString));
+			FSlateApplication::Get().SetKeyboardFocus(EditableTextBox);
+			if(PlaceHolderComboButton.IsValid())
+			{
+				PlaceHolderComboButton->SetIsOpen(false);
+			}
+		}
+	}
+}
+
+void FMetaDataTextKeyPatternStructCustomization::CustomizeStructHeader( TSharedRef<IPropertyHandle> StructPropertyHandle, class FDetailWidgetRow& HeaderRow, IStructCustomizationUtils& StructCustomizationUtils )
+{
+	const TSharedPtr<IPropertyHandle> PatternPropertyHandle = StructPropertyHandle->GetChildHandle(GET_MEMBER_NAME_CHECKED(FMetaDataTextKeyPattern, Pattern));
+
+	const auto& GetErrorText = [StructPropertyHandle]() -> FText
+	{
+		if (StructPropertyHandle->IsValidHandle())
+		{
+			TArray<const void*> RawData;
+			StructPropertyHandle->AccessRawData(RawData);
+			const FMetaDataTextKeyPattern& Struct = *(reinterpret_cast<const FMetaDataTextKeyPattern*>(RawData.Top()));
+
+			FText Error;
+			if (!Struct.Validate(Error))
+			{
+				return Error;
+			}
+		}
+
+		return FText::GetEmpty();
+	};
+
+	HeaderRow
+		.NameContent()
+		[
+			SNew(SHorizontalBox)
+			+SHorizontalBox::Slot()
+			.AutoWidth()
+			[
+				StructPropertyHandle->CreatePropertyNameWidget()
+			]
+			+SHorizontalBox::Slot()
+			.AutoWidth()
+			[
+				SNew(SGatherTextConfigurationErrorHint)
+				.ErrorText_Lambda(GetErrorText)
+			]
+		]
+		.ValueContent()
+		[
+			SNew(SMetaDataTextKeyPatternWidget, PatternPropertyHandle.ToSharedRef())
+		];
+}
+
+void FMetaDataKeyNameStructCustomization::CustomizeStructHeader( TSharedRef<IPropertyHandle> StructPropertyHandle, class FDetailWidgetRow& HeaderRow, IStructCustomizationUtils& StructCustomizationUtils )
+{
+	const TSharedPtr<IPropertyHandle> NamePropertyHandle = StructPropertyHandle->GetChildHandle(GET_MEMBER_NAME_CHECKED(FMetaDataKeyName, Name));
+
+	const auto& GetErrorText = [StructPropertyHandle]() -> FText
+	{
+		if (StructPropertyHandle->IsValidHandle())
+		{
+			TArray<const void*> RawData;
+			StructPropertyHandle->AccessRawData(RawData);
+			const FMetaDataKeyName& Struct = *(reinterpret_cast<const FMetaDataKeyName*>(RawData.Top()));
+
+			FText Error;
+			if (!Struct.Validate(Error))
+			{
+				return Error;
+			}
+		}
+
+		return FText::GetEmpty();
+	};
+
+	HeaderRow
+		.NameContent()
+		[
+			SNew(SHorizontalBox)
+			+SHorizontalBox::Slot()
+			.AutoWidth()
+			[
+				StructPropertyHandle->CreatePropertyNameWidget()
+			]
+			+SHorizontalBox::Slot()
+			.AutoWidth()
+			[
+				SNew(SGatherTextConfigurationErrorHint)
+				.ErrorText_Lambda(GetErrorText)
+			]
+		]
+		.ValueContent()
+		[
+			NamePropertyHandle->CreatePropertyValueWidget()
+		];
+}
+
+void FMetaDataKeyGatherSpecificationStructCustomization::CustomizeStructHeader(TSharedRef<IPropertyHandle> StructPropertyHandle, class FDetailWidgetRow& HeaderRow, IStructCustomizationUtils& StructCustomizationUtils)
+{
+	const auto& GetErrorText = [StructPropertyHandle]() -> FText
+	{
+		if (StructPropertyHandle->IsValidHandle())
+		{
+			TArray<const void*> RawData;
+			StructPropertyHandle->AccessRawData(RawData);
+			const FMetaDataKeyGatherSpecification& Struct = *(reinterpret_cast<const FMetaDataKeyGatherSpecification*>(RawData.Top()));
+
+			FText Error;
+			if (!Struct.Validate(Error))
+			{
+				return Error;
+			}
+		}
+
+		return FText::GetEmpty();
+	};
+
+	HeaderRow
+		.NameContent()
+		[
+			SNew(SHorizontalBox)
+			+SHorizontalBox::Slot()
+			.AutoWidth()
+			[
+				StructPropertyHandle->CreatePropertyNameWidget()
+			]
+			+SHorizontalBox::Slot()
+			.AutoWidth()
+			[
+				SNew(SGatherTextConfigurationErrorHint)
+				.ErrorText_Lambda(GetErrorText)
+			]
+		]
+		.ValueContent()
+		[
+			StructPropertyHandle->CreatePropertyValueWidget()
+		];
+}
+
+void FMetaDataKeyGatherSpecificationStructCustomization::CustomizeStructChildren( TSharedRef<IPropertyHandle> StructPropertyHandle, class IDetailChildrenBuilder& ChildBuilder, IStructCustomizationUtils& StructCustomizationUtils )
+{
+	uint32 ChildCount;
+	StructPropertyHandle->GetNumChildren(ChildCount);
+
+	for (uint32 i = 0; i < ChildCount; ++i)
+	{
+		const TSharedPtr<IPropertyHandle> ChildPropertyHandle = StructPropertyHandle->GetChildHandle(i);
+		if (!ChildPropertyHandle->IsCustomized())
+		{
+			ChildBuilder.AddChildProperty(ChildPropertyHandle.ToSharedRef());
+		}
+	}
+}
+
+void FGatherTextFromMetaDataConfigurationStructCustomization::CustomizeStructHeader( TSharedRef<IPropertyHandle> StructPropertyHandle, class FDetailWidgetRow& HeaderRow, IStructCustomizationUtils& StructCustomizationUtils )
+{
+	const TSharedPtr<IPropertyHandle> IsEnabledPropertyHandle = StructPropertyHandle->GetChildHandle(GET_MEMBER_NAME_CHECKED(FGatherTextFromMetaDataConfiguration, IsEnabled));
+	if (IsEnabledPropertyHandle.IsValid() && IsEnabledPropertyHandle->IsValidHandle())
+	{
+		IsEnabledPropertyHandle->MarkHiddenByCustomization();
+	}
+
+	const auto& GetConfigurationError = [StructPropertyHandle]() -> FText
+	{
+		TArray<const void*> RawData;
+		StructPropertyHandle->AccessRawData(RawData);
+		const FGatherTextFromMetaDataConfiguration& Struct = *(reinterpret_cast<const FGatherTextFromMetaDataConfiguration*>(RawData.Top()));
+
+		TArray<UObject*> OuterObjects;
+		StructPropertyHandle->GetOuterObjects(OuterObjects);
+		ULocalizationTarget* const LocalizationTarget = Cast<ULocalizationTarget>(OuterObjects.Top());
+		if (LocalizationTarget)
+		{
+			const FString DesiredRootPath = FPaths::ConvertRelativePathToFull(LocalizationTarget->IsMemberOfEngineTargetSet() ? FPaths::EngineDir() : FPaths::GameDir());
+			FText Error;
+			if (!Struct.Validate(DesiredRootPath, Error))
+			{
+				return Error;
+			}
+		}
+
+		return FText::GetEmpty();
+	};
+
+	const auto& GetValidityWidgetVisibility = [IsEnabledPropertyHandle]() -> EVisibility
+	{
+		bool Value = false;
+		if (IsEnabledPropertyHandle.IsValid() && IsEnabledPropertyHandle->IsValidHandle())
+		{
+			IsEnabledPropertyHandle->GetValue(Value);
+		}
+		return Value ? EVisibility::Visible : EVisibility::Hidden;
+	};
+		
+	HeaderRow
+		.NameContent()
+		[
+			SNew(SHorizontalBox)
+			+SHorizontalBox::Slot()
+			.AutoWidth()
+			.VAlign(VAlign_Center)
+			.Padding(2.0f, 0.0f)
+			[
+				IsEnabledPropertyHandle.IsValid() && IsEnabledPropertyHandle->IsValidHandle() ? IsEnabledPropertyHandle->CreatePropertyValueWidget() : SNullWidget::NullWidget
+			]
+			+SHorizontalBox::Slot()
+			.AutoWidth()
+			.VAlign(VAlign_Center)
+			[
+				StructPropertyHandle->CreatePropertyNameWidget()
+			]
+		]
+		.ValueContent()
+		[
+			SNew(SHorizontalBox)
+			+SHorizontalBox::Slot()
+			.AutoWidth()
+			[
+				StructPropertyHandle->CreatePropertyValueWidget()
+			]
+			+SHorizontalBox::Slot()
+			.AutoWidth()
+			[
+				SNew(SConfigurationValidity)
+				.Visibility_Lambda(GetValidityWidgetVisibility)
+				.ConfigurationError_Lambda(GetConfigurationError)
+			]
+		];
+}
+
+void FGatherTextFromMetaDataConfigurationStructCustomization::CustomizeStructChildren( TSharedRef<IPropertyHandle> StructPropertyHandle, class IDetailChildrenBuilder& ChildBuilder, IStructCustomizationUtils& StructCustomizationUtils )
+{
+	uint32 ChildCount;
+	StructPropertyHandle->GetNumChildren(ChildCount);
+
+	for (uint32 i = 0; i < ChildCount; ++i)
+	{
+		const TSharedPtr<IPropertyHandle> ChildPropertyHandle = StructPropertyHandle->GetChildHandle(i);
+		if (!ChildPropertyHandle->IsCustomized())
+		{
+			ChildBuilder.AddChildProperty(ChildPropertyHandle.ToSharedRef());
+		}
+	}
+}
+
+#undef LOCTEXT_NAMESPACE
\ No newline at end of file
diff --git a/Engine/Source/Editor/LocalizationDashboard/Private/GatherTextDetailCustomizations.h b/Engine/Source/Editor/LocalizationDashboard/Private/GatherTextDetailCustomizations.h
new file mode 100644
--- /dev/null
+++ b/Engine/Source/Editor/LocalizationDashboard/Private/GatherTextDetailCustomizations.h
@@ -0,0 +1,76 @@
+// Copyright 1998-2015 Epic Games, Inc. All Rights Reserved.
+
+#pragma once
+
+#include "IPropertyTypeCustomization.h"
+#include "LocalizationTargetTypes.h"
+
+class FGatherTextSearchDirectoryStructCustomization : public IStructCustomization
+{
+public:
+	virtual void CustomizeStructHeader(TSharedRef<IPropertyHandle> StructPropertyHandle, class FDetailWidgetRow& HeaderRow, IStructCustomizationUtils& StructCustomizationUtils) override;
+	virtual void CustomizeStructChildren(TSharedRef<IPropertyHandle> StructPropertyHandle, class IDetailChildrenBuilder& ChildBuilder, IStructCustomizationUtils& StructCustomizationUtils) override {};
+};
+
+class FGatherTextIncludePathStructCustomization : public IStructCustomization
+{
+public:
+	virtual void CustomizeStructHeader(TSharedRef<IPropertyHandle> StructPropertyHandle, class FDetailWidgetRow& HeaderRow, IStructCustomizationUtils& StructCustomizationUtils) override;
+	virtual void CustomizeStructChildren(TSharedRef<IPropertyHandle> StructPropertyHandle, class IDetailChildrenBuilder& ChildBuilder, IStructCustomizationUtils& StructCustomizationUtils) override {};
+};
+
+class FGatherTextExcludePathStructCustomization : public IStructCustomization
+{
+public:
+	virtual void CustomizeStructHeader(TSharedRef<IPropertyHandle> StructPropertyHandle, class FDetailWidgetRow& HeaderRow, IStructCustomizationUtils& StructCustomizationUtils) override;
+	virtual void CustomizeStructChildren(TSharedRef<IPropertyHandle> StructPropertyHandle, class IDetailChildrenBuilder& ChildBuilder, IStructCustomizationUtils& StructCustomizationUtils) override {};
+};
+
+class FGatherTextFileExtensionStructCustomization : public IStructCustomization
+{
+public:
+	virtual void CustomizeStructHeader(TSharedRef<IPropertyHandle> StructPropertyHandle, class FDetailWidgetRow& HeaderRow, IStructCustomizationUtils& StructCustomizationUtils) override;
+	virtual void CustomizeStructChildren(TSharedRef<IPropertyHandle> StructPropertyHandle, class IDetailChildrenBuilder& ChildBuilder, IStructCustomizationUtils& StructCustomizationUtils) override {};
+};
+
+class FGatherTextFromTextFilesConfigurationStructCustomization : public IStructCustomization
+{
+public:
+	virtual void CustomizeStructHeader(TSharedRef<IPropertyHandle> StructPropertyHandle, class FDetailWidgetRow& HeaderRow, IStructCustomizationUtils& StructCustomizationUtils) override;
+	virtual void CustomizeStructChildren(TSharedRef<IPropertyHandle> StructPropertyHandle, class IDetailChildrenBuilder& ChildBuilder, IStructCustomizationUtils& StructCustomizationUtils) override;
+};
+
+class FGatherTextFromPackagesConfigurationStructCustomization : public IStructCustomization
+{
+public:
+	virtual void CustomizeStructHeader(TSharedRef<IPropertyHandle> StructPropertyHandle, class FDetailWidgetRow& HeaderRow, IStructCustomizationUtils& StructCustomizationUtils) override;
+	virtual void CustomizeStructChildren(TSharedRef<IPropertyHandle> StructPropertyHandle, class IDetailChildrenBuilder& ChildBuilder, IStructCustomizationUtils& StructCustomizationUtils) override;
+};
+
+class FMetaDataTextKeyPatternStructCustomization : public IStructCustomization
+{
+public:
+	virtual void CustomizeStructHeader(TSharedRef<IPropertyHandle> StructPropertyHandle, class FDetailWidgetRow& HeaderRow, IStructCustomizationUtils& StructCustomizationUtils) override;
+	virtual void CustomizeStructChildren(TSharedRef<IPropertyHandle> StructPropertyHandle, class IDetailChildrenBuilder& ChildBuilder, IStructCustomizationUtils& StructCustomizationUtils) override {};
+};
+
+class FMetaDataKeyNameStructCustomization : public IStructCustomization
+{
+public:
+	virtual void CustomizeStructHeader(TSharedRef<IPropertyHandle> StructPropertyHandle, class FDetailWidgetRow& HeaderRow, IStructCustomizationUtils& StructCustomizationUtils) override;
+	virtual void CustomizeStructChildren(TSharedRef<IPropertyHandle> StructPropertyHandle, class IDetailChildrenBuilder& ChildBuilder, IStructCustomizationUtils& StructCustomizationUtils) override {};
+};
+
+class FMetaDataKeyGatherSpecificationStructCustomization : public IStructCustomization
+{
+public:
+	virtual void CustomizeStructHeader(TSharedRef<IPropertyHandle> StructPropertyHandle, class FDetailWidgetRow& HeaderRow, IStructCustomizationUtils& StructCustomizationUtils) override;
+	virtual void CustomizeStructChildren(TSharedRef<IPropertyHandle> StructPropertyHandle, class IDetailChildrenBuilder& ChildBuilder, IStructCustomizationUtils& StructCustomizationUtils) override;
+};
+
+class FGatherTextFromMetaDataConfigurationStructCustomization : public IStructCustomization
+{
+public:
+	virtual void CustomizeStructHeader(TSharedRef<IPropertyHandle> StructPropertyHandle, class FDetailWidgetRow& HeaderRow, IStructCustomizationUtils& StructCustomizationUtils) override;
+	virtual void CustomizeStructChildren(TSharedRef<IPropertyHandle> StructPropertyHandle, class IDetailChildrenBuilder& ChildBuilder, IStructCustomizationUtils& StructCustomizationUtils) override;
+};
diff --git a/Engine/Source/Editor/LocalizationDashboard/Private/LocalizationConfigurationScript.cpp b/Engine/Source/Editor/LocalizationDashboard/Private/LocalizationConfigurationScript.cpp
--- a/Engine/Source/Editor/LocalizationDashboard/Private/LocalizationConfigurationScript.cpp
+++ b/Engine/Source/Editor/LocalizationDashboard/Private/LocalizationConfigurationScript.cpp
@@ -140,7 +140,7 @@ namespace LocalizationConfigurationScript
 
 		uint32 GatherTextStepIndex = 0;
 		// GatherTextFromSource
-		if (Target->Settings.GatherFromTextFiles.SearchDirectories.Num() && Target->Settings.GatherFromTextFiles.FileExtensions.Num()) // Don't gather from source if there are no source files.
+		if (Target->Settings.GatherFromTextFiles.IsEnabled)
 		{
 			FConfigSection& ConfigSection = Script.GatherTextStep(GatherTextStepIndex++);
 
@@ -150,25 +150,25 @@ namespace LocalizationConfigurationScript
 			// Include Paths
 			for (const auto& IncludePath : Target->Settings.GatherFromTextFiles.SearchDirectories)
 			{
-				ConfigSection.Add( TEXT("IncludePaths"), IncludePath );
+				ConfigSection.Add( TEXT("SearchDirectoryPaths"), IncludePath.Path );
 			}
 
 			// Exclude Paths
-			ConfigSection.Add( TEXT("ExcludePaths"), TEXT("*/Config/Localization/*") );
+			ConfigSection.Add( TEXT("ExcludePathFilters"), TEXT("Config/Localization/*") );
 			for (const auto& ExcludePath : Target->Settings.GatherFromTextFiles.ExcludePathWildcards)
 			{
-				ConfigSection.Add( TEXT("ExcludePaths"), ExcludePath );
+				ConfigSection.Add( TEXT("ExcludePathFilters"), ExcludePath.Pattern );
 			}
 
 			// Source File Search Filters
 			for (const auto& FileExtension : Target->Settings.GatherFromTextFiles.FileExtensions)
 			{
-				ConfigSection.Add( TEXT("SourceFileSearchFilters"), FileExtension );
+				ConfigSection.Add( TEXT("FileNameFilters"), FString::Printf( TEXT("*.%s"), *FileExtension.Pattern) );
 			}
 		}
 
 		// GatherTextFromAssets
-		if (Target->Settings.GatherFromPackages.IncludePathWildcards.Num() && Target->Settings.GatherFromPackages.FileExtensions.Num()) // Don't gather from packages if there are none specified.
+		if (Target->Settings.GatherFromPackages.IsEnabled)
 		{
 			FConfigSection& ConfigSection = Script.GatherTextStep(GatherTextStepIndex++);
 
@@ -178,25 +178,25 @@ namespace LocalizationConfigurationScript
 			// Include Paths
 			for (const auto& IncludePath : Target->Settings.GatherFromPackages.IncludePathWildcards)
 			{
-				ConfigSection.Add( TEXT("IncludePaths"), IncludePath );
+				ConfigSection.Add( TEXT("IncludePathFilters"), IncludePath.Pattern );
 			}
 
 			// Exclude Paths
-			ConfigSection.Add( TEXT("ExcludePaths"), TEXT("*/Content/Localization/*") );
+			ConfigSection.Add( TEXT("ExcludePathFilters"), TEXT("Content/Localization/*") );
 			for (const auto& ExcludePath : Target->Settings.GatherFromPackages.ExcludePathWildcards)
 			{
-				ConfigSection.Add( TEXT("ExcludePaths"), ExcludePath );
+				ConfigSection.Add( TEXT("ExcludePathFilters"), ExcludePath.Pattern );
 			}
 
 			// Package Extensions
 			for (const auto& FileExtension : Target->Settings.GatherFromPackages.FileExtensions)
 			{
-				ConfigSection.Add( TEXT("PackageExtensions"), FileExtension );
+				ConfigSection.Add( TEXT("PackageFileNameFilters"), FString::Printf( TEXT("*.%s"), *FileExtension.Pattern) );
 			}
 		}
 
 		// GatherTextFromMetadata
-		if (Target->Settings.GatherFromMetaData.IncludePathWildcards.Num() && Target->Settings.GatherFromMetaData.KeySpecifications.Num()) // Don't gather from metadata if there are none specified.
+		if (Target->Settings.GatherFromMetaData.IsEnabled)
 		{
 			FConfigSection& ConfigSection = Script.GatherTextStep(GatherTextStepIndex++);
 
@@ -206,21 +206,21 @@ namespace LocalizationConfigurationScript
 			// Include Paths
 			for (const auto& IncludePath : Target->Settings.GatherFromMetaData.IncludePathWildcards)
 			{
-				ConfigSection.Add( TEXT("IncludePaths"), IncludePath );
+				ConfigSection.Add( TEXT("IncludePathFilters"), IncludePath.Pattern );
 			}
 
 			// Exclude Paths
 			for (const auto& ExcludePath : Target->Settings.GatherFromMetaData.ExcludePathWildcards)
 			{
-				ConfigSection.Add( TEXT("ExcludePaths"), ExcludePath );
+				ConfigSection.Add( TEXT("ExcludePathFilters"), ExcludePath.Pattern );
 			}
 
 			// Package Extensions
 			for (const FMetaDataKeyGatherSpecification& Specification : Target->Settings.GatherFromMetaData.KeySpecifications)
 			{
-				ConfigSection.Add( TEXT("InputKeys"), Specification.MetaDataKey );
+				ConfigSection.Add( TEXT("InputKeys"), Specification.MetaDataKey.Name );
 				ConfigSection.Add( TEXT("OutputNamespaces"), Specification.TextNamespace );
-				ConfigSection.Add( TEXT("OutputKeys"), FString::Printf(TEXT("\"%s\""), *Specification.TextKeyPattern) );
+				ConfigSection.Add( TEXT("OutputKeys"), FString::Printf(TEXT("\"%s\""), *Specification.TextKeyPattern.Pattern) );
 			}
 		}
 
@@ -247,24 +247,10 @@ namespace LocalizationConfigurationScript
 			// CommandletClass
 			ConfigSection.Add( TEXT("CommandletClass"), TEXT("GenerateTextLocalizationReport") );
 
-			const FString SourcePath = ContentDirRelativeToGameDir / TEXT("Localization") / Target->Settings.Name;
-			ConfigSection.Add( TEXT("SourcePath"), SourcePath );
-			const FString DestinationPath = ContentDirRelativeToGameDir / TEXT("Localization") / Target->Settings.Name;
-			ConfigSection.Add( TEXT("DestinationPath"), DestinationPath );
-
-			ConfigSection.Add( TEXT("ManifestName"), FString::Printf( TEXT("%s.%s"), *Target->Settings.Name, TEXT("manifest") ) );
-
-			for (const FCultureStatistics& CultureStatistics : Target->Settings.SupportedCulturesStatistics)
-			{
-				ConfigSection.Add( TEXT("CulturesToGenerate"), CultureStatistics.CultureName );
-			}
-
 			ConfigSection.Add( TEXT("bWordCountReport"), TEXT("true") );
-
 			ConfigSection.Add( TEXT("WordCountReportName"), FPaths::GetCleanFilename( GetWordCountCSVPath(Target) ) );
 
 			ConfigSection.Add( TEXT("bConflictReport"), TEXT("true") );
-
 			ConfigSection.Add( TEXT("ConflictReportName"), FPaths::GetCleanFilename( GetConflictReportPath(Target) ) );
 		}
 
diff --git a/Engine/Source/Editor/LocalizationDashboard/Private/LocalizationDashboardModule.cpp b/Engine/Source/Editor/LocalizationDashboard/Private/LocalizationDashboardModule.cpp
--- a/Engine/Source/Editor/LocalizationDashboard/Private/LocalizationDashboardModule.cpp
+++ b/Engine/Source/Editor/LocalizationDashboard/Private/LocalizationDashboardModule.cpp
@@ -9,6 +9,7 @@
 #include "ILocalizationServiceProvider.h"
 #include "LocalizationTargetSetDetailCustomization.h"
 #include "LocalizationTargetDetailCustomization.h"
+#include "GatherTextDetailCustomizations.h"
 
 #define LOCTEXT_NAMESPACE "LocalizationDashboard"
 
@@ -35,7 +36,67 @@ public:
 				return MakeShareable( new FLocalizationTargetDetailCustomization());
 			})
 		);
-
+		PropertyModule.RegisterCustomPropertyTypeLayout("GatherTextSearchDirectory", FOnGetPropertyTypeCustomizationInstance::CreateLambda(
+			[]() -> TSharedRef<IStructCustomization>
+			{
+				return MakeShareable( new FGatherTextSearchDirectoryStructCustomization());
+			})
+		);
+		PropertyModule.RegisterCustomPropertyTypeLayout("GatherTextIncludePath", FOnGetPropertyTypeCustomizationInstance::CreateLambda(
+			[]() -> TSharedRef<IStructCustomization>
+			{
+				return MakeShareable( new FGatherTextIncludePathStructCustomization());
+			})
+		);
+		PropertyModule.RegisterCustomPropertyTypeLayout("GatherTextExcludePath", FOnGetPropertyTypeCustomizationInstance::CreateLambda(
+			[]() -> TSharedRef<IStructCustomization>
+			{
+				return MakeShareable( new FGatherTextExcludePathStructCustomization());
+			})
+		);
+		PropertyModule.RegisterCustomPropertyTypeLayout("GatherTextFileExtension", FOnGetPropertyTypeCustomizationInstance::CreateLambda(
+			[]() -> TSharedRef<IStructCustomization>
+			{
+				return MakeShareable( new FGatherTextFileExtensionStructCustomization());
+			})
+		);
+		PropertyModule.RegisterCustomPropertyTypeLayout("GatherTextFromTextFilesConfiguration", FOnGetPropertyTypeCustomizationInstance::CreateLambda(
+			[]() -> TSharedRef<IStructCustomization>
+			{
+				return MakeShareable( new FGatherTextFromTextFilesConfigurationStructCustomization());
+			})
+		);
+		PropertyModule.RegisterCustomPropertyTypeLayout("GatherTextFromPackagesConfiguration", FOnGetPropertyTypeCustomizationInstance::CreateLambda(
+			[]() -> TSharedRef<IStructCustomization>
+			{
+				return MakeShareable( new FGatherTextFromPackagesConfigurationStructCustomization());
+			})
+		);
+		PropertyModule.RegisterCustomPropertyTypeLayout("MetaDataTextKeyPattern", FOnGetPropertyTypeCustomizationInstance::CreateLambda(
+			[]() -> TSharedRef<IStructCustomization>
+			{
+				return MakeShareable( new FMetaDataTextKeyPatternStructCustomization());
+			})
+		);
+		PropertyModule.RegisterCustomPropertyTypeLayout("MetaDataKeyName", FOnGetPropertyTypeCustomizationInstance::CreateLambda(
+			[]() -> TSharedRef<IStructCustomization>
+			{
+				return MakeShareable( new FMetaDataKeyNameStructCustomization());
+			})
+		);
+		PropertyModule.RegisterCustomPropertyTypeLayout("MetaDataKeyGatherSpecification", FOnGetPropertyTypeCustomizationInstance::CreateLambda(
+			[]() -> TSharedRef<IStructCustomization>
+			{
+				return MakeShareable( new FMetaDataKeyGatherSpecificationStructCustomization());
+			})
+		);
+		PropertyModule.RegisterCustomPropertyTypeLayout("GatherTextFromMetaDataConfiguration", FOnGetPropertyTypeCustomizationInstance::CreateLambda(
+			[]() -> TSharedRef<IStructCustomization>
+			{
+				return MakeShareable( new FGatherTextFromMetaDataConfigurationStructCustomization());
+			})
+		);
+			
 		FLocalizationDashboard::Initialize();
 	}
 
@@ -44,6 +105,16 @@ public:
 		FLocalizationDashboard::Terminate();
 
 		FPropertyEditorModule& PropertyModule = FModuleManager::LoadModuleChecked<FPropertyEditorModule>("PropertyEditor");
+		PropertyModule.UnregisterCustomPropertyTypeLayout("GatherTextFromMetaDataConfiguration");
+		PropertyModule.UnregisterCustomPropertyTypeLayout("MetaDataKeyGatherSpecification");
+		PropertyModule.UnregisterCustomPropertyTypeLayout("MetaDataKeyName");
+		PropertyModule.UnregisterCustomPropertyTypeLayout("MetaDataTextKeyPattern");
+		PropertyModule.UnregisterCustomPropertyTypeLayout("GatherTextFromPackagesConfiguration");
+		PropertyModule.UnregisterCustomPropertyTypeLayout("GatherTextFromTextFilesConfiguration");
+		PropertyModule.UnregisterCustomPropertyTypeLayout("GatherTextFileExtension");
+		PropertyModule.UnregisterCustomPropertyTypeLayout("GatherTextExcludePath");
+		PropertyModule.UnregisterCustomPropertyTypeLayout("GatherTextIncludePath");
+		PropertyModule.UnregisterCustomPropertyTypeLayout("GatherTextSearchDirectory");
 		PropertyModule.UnregisterCustomClassLayout("LocalizationTarget");
 		PropertyModule.UnregisterCustomClassLayout("LocalizationTargetSet");
 	}
diff --git a/Engine/Source/Editor/LocalizationDashboard/Private/LocalizationTargetTypes.cpp b/Engine/Source/Editor/LocalizationDashboard/Private/LocalizationTargetTypes.cpp
--- a/Engine/Source/Editor/LocalizationDashboard/Private/LocalizationTargetTypes.cpp
+++ b/Engine/Source/Editor/LocalizationDashboard/Private/LocalizationTargetTypes.cpp
@@ -6,6 +6,219 @@
 #include "LocalizationConfigurationScript.h"
 #include "LocalizationDashboardSettings.h"
 
+#define LOCTEXT_NAMESPACE "LocalizationTargetTypes"
+
+bool FGatherTextSearchDirectory::Validate(const FString& RootDirectory, FText& OutError) const
+{
+	if (Path.IsEmpty())
+	{
+		OutError = LOCTEXT("SearchDirectoryEmptyError", "Search directory not specified. Use \".\" to specify the root directory.");
+		return false;
+	}
+
+	FText InvalidPathReason;
+	if (!FPaths::ValidatePath(Path, &InvalidPathReason))
+	{
+		OutError = InvalidPathReason;
+		return false;
+	}
+
+	if (!FPaths::DirectoryExists(FPaths::Combine(*RootDirectory, *Path)))
+	{
+		OutError = LOCTEXT("SearchDirectoryNonExistentError", "Search directory does not exist.");
+		return false;
+	}
+
+	return true;
+}
+
+bool FGatherTextIncludePath::Validate(const FString& RootDirectory, FText& OutError) const
+{
+	if (Pattern.IsEmpty())
+	{
+		OutError = LOCTEXT("IncludePathEmptyError", "Include path not specified. Use \".\" to specify the root directory.");
+		return false;
+	}
+
+	if (!Pattern.Contains(TEXT("*")))
+	{
+		OutError = LOCTEXT("IncludePathNoWildcardError", "Include path does not specify a wild card (\"*\"). Append \"*\" or only the file at the exact specified directory will be gathered from.");
+		return false;
+	}
+
+	return true;
+}
+
+bool FGatherTextExcludePath::Validate(FText& OutError) const
+{
+	if (Pattern.IsEmpty())
+	{
+		OutError = LOCTEXT("ExcludePathEmptyError", "Exclude path not specified. Use \".\" to specify the root directory.");
+		return false;
+	}
+
+	if (!Pattern.Contains(TEXT("*")))
+	{
+		OutError = LOCTEXT("ExcludePathNoWildcardError", "Exclude path does not specify a wild card (\"*\"). Append \"*\" or only the file at the exact specified directory will be excluded.");
+		return false;
+	}
+
+	return true;
+}
+
+bool FGatherTextFileExtension::Validate(FText& OutError) const
+{
+	if (Pattern.IsEmpty())
+	{
+		OutError = LOCTEXT("FileExtensionEmptyError", "File extension not specified.");
+		return false;
+	}
+
+	return true;
+}
+
+const TArray<FGatherTextFileExtension> FGatherTextFromTextFilesConfiguration::DefaultTextFileExtensions = []()
+{
+	TArray<FGatherTextFileExtension> Result;
+	Result.Add(FGatherTextFileExtension{TEXT("h")});
+	Result.Add(FGatherTextFileExtension{TEXT("cpp")});
+	Result.Add(FGatherTextFileExtension{TEXT("ini")});
+	return Result;
+}();
+
+const TArray<FGatherTextFileExtension> FGatherTextFromPackagesConfiguration::DefaultPackageFileExtensions = []()
+{
+	TArray<FGatherTextFileExtension> Result;
+	Result.Add(FGatherTextFileExtension{TEXT("umap")});
+	Result.Add(FGatherTextFileExtension{TEXT("uasset")});
+	return Result;
+}();
+
+bool FGatherTextFromTextFilesConfiguration::Validate(const FString& RootDirectory, FText& OutError) const
+{
+	const auto& InvalidSearchDirectoryPredicate = [RootDirectory](const FGatherTextSearchDirectory& Element) -> bool
+	{
+		FText ElementError;
+		return !(Element.Validate(RootDirectory, ElementError));
+	};
+	const auto& InvalidExcludePathPredicate = [RootDirectory](const FGatherTextExcludePath& Element) -> bool
+	{
+		FText ElementError;
+		return !(Element.Validate(ElementError));
+	};
+	const auto& InvalidFileExtensionPredicate = [RootDirectory](const FGatherTextFileExtension& Element) -> bool
+	{
+		FText ElementError;
+		return !(Element.Validate(ElementError));
+	};
+
+	if (SearchDirectories.Num() > 0 && !SearchDirectories.ContainsByPredicate(InvalidSearchDirectoryPredicate) &&
+		!ExcludePathWildcards.ContainsByPredicate(InvalidExcludePathPredicate) &&
+		FileExtensions.Num() > 0 && !FileExtensions.ContainsByPredicate(InvalidFileExtensionPredicate))
+	{
+		return true;
+	}
+
+	OutError = LOCTEXT("InvalidGatherTextFromFilesConfigurationError", "Must have at least one search directory, one file extension, and no invalid settings.");
+	return false;
+}
+
+bool FGatherTextFromPackagesConfiguration::Validate(const FString& RootDirectory, FText& OutError) const
+{
+	const auto& InvalidIncludePathPredicate = [RootDirectory](const FGatherTextIncludePath& Element) -> bool
+	{
+		FText ElementError;
+		return !(Element.Validate(RootDirectory, ElementError));
+	};
+	const auto& InvalidExcludePathPredicate = [RootDirectory](const FGatherTextExcludePath& Element) -> bool
+	{
+		FText ElementError;
+		return !(Element.Validate(ElementError));
+	};
+	const auto& InvalidFileExtensionPredicate = [RootDirectory](const FGatherTextFileExtension& Element) -> bool
+	{
+		FText ElementError;
+		return !(Element.Validate(ElementError));
+	};
+
+	if (IncludePathWildcards.Num() > 0 && !IncludePathWildcards.ContainsByPredicate(InvalidIncludePathPredicate) &&
+		!ExcludePathWildcards.ContainsByPredicate(InvalidExcludePathPredicate) &&
+		FileExtensions.Num() > 0 && !FileExtensions.ContainsByPredicate(InvalidFileExtensionPredicate))
+	{
+		return true;
+	}
+
+	OutError = LOCTEXT("InvalidGatherTextFromPackagesConfigurationError", "Must have at least one include path, one file extension, and no invalid settings.");
+	return false;
+}
+
+const TArray<FString> FMetaDataTextKeyPattern::PossiblePlaceHolders = []() -> TArray<FString>
+{
+	TArray<FString> Result;
+	Result.Add(TEXT("{FieldPath}"));
+	Result.Add(TEXT("{MetaDataValue}"));
+	return Result;
+}();
+
+bool FMetaDataTextKeyPattern::Validate(FText& OutError) const
+{
+	for (const auto& PossiblePlaceHolder : PossiblePlaceHolders)
+	{
+		if (Pattern.Contains(PossiblePlaceHolder))
+		{
+			return true;
+		}
+	}
+
+	OutError = LOCTEXT("NoMetaDataLocalizationKeyPlaceHolderError", "No place holders used. All generated keys will conflict!");
+	return false;
+}
+
+bool FMetaDataKeyName::Validate(FText& OutError) const
+{
+	if (Name.IsEmpty())
+	{
+		OutError = LOCTEXT("MetaDataKeyNameEmptyError", "Meta data key not specified.");
+		return false;
+	}
+
+	return true;
+}
+
+bool FMetaDataKeyGatherSpecification::Validate(FText& OutError) const
+{
+	return MetaDataKey.Validate(OutError) && TextKeyPattern.Validate(OutError);
+}
+
+bool FGatherTextFromMetaDataConfiguration::Validate(const FString& RootDirectory, FText& OutError) const
+{
+	const auto& InvalidIncludePathPredicate = [RootDirectory](const FGatherTextIncludePath& Element) -> bool
+	{
+		FText ElementError;
+		return !(Element.Validate(RootDirectory, ElementError));
+	};
+	const auto& InvalidExcludePathPredicate = [RootDirectory](const FGatherTextExcludePath& Element) -> bool
+	{
+		FText ElementError;
+		return !(Element.Validate(ElementError));
+	};
+	const auto& InvalidMetaDataKeySpecificationPredicate = [RootDirectory](const FMetaDataKeyGatherSpecification& Element) -> bool
+	{
+		FText ElementError;
+		return !(Element.Validate(ElementError));
+	};
+
+	if (IncludePathWildcards.Num() > 0 && !IncludePathWildcards.ContainsByPredicate(InvalidIncludePathPredicate) &&
+		!ExcludePathWildcards.ContainsByPredicate(InvalidExcludePathPredicate) &&
+		KeySpecifications.Num() > 0 && !KeySpecifications.ContainsByPredicate(InvalidMetaDataKeySpecificationPredicate))
+	{
+		return true;
+	}
+
+	OutError = LOCTEXT("InvalidGatherTextFromMetadataConfigurationError", "Must have at least one include path, one meta data key specification, and no invalid settings.");
+	return false;
+}
+
 FLocalizationTargetSettings::FLocalizationTargetSettings()
 	: Guid(FGuid::NewGuid())
 	, ConflictStatus(ELocalizationTargetConflictStatus::Unknown)
@@ -457,4 +670,6 @@ void ULocalizationTargetSet::PostEditChangeProperty(FPropertyChangedEvent& Prope
 	{
 		LocalizationDashboardSettings->PostEditChange();
 	}
-}
\ No newline at end of file
+}
+
+#undef LOCTEXT_NAMESPACE
\ No newline at end of file
diff --git a/Engine/Source/Editor/LocalizationDashboard/Private/LocalizationTargetTypes.h b/Engine/Source/Editor/LocalizationDashboard/Private/LocalizationTargetTypes.h
--- a/Engine/Source/Editor/LocalizationDashboard/Private/LocalizationTargetTypes.h
+++ b/Engine/Source/Editor/LocalizationDashboard/Private/LocalizationTargetTypes.h
@@ -13,40 +13,137 @@ enum class ELocalizationTargetConflictStatus : uint8
 	Clear, /* The localization target is clear of conflicts. */
 };
 
+USTRUCT()
+struct FGatherTextSearchDirectory
+{
+	GENERATED_USTRUCT_BODY()
+
+	UPROPERTY(config, EditAnywhere, Category="Path")
+	FString Path;
+
+	bool Validate(const FString& RootDirectory, FText& OutError) const;
+};
+
+USTRUCT()
+struct FGatherTextIncludePath
+{
+	GENERATED_USTRUCT_BODY()
+
+	UPROPERTY(config, EditAnywhere, Category="Pattern")
+	FString Pattern;
+
+	bool Validate(const FString& RootDirectory, FText& OutError) const;
+};
+
+USTRUCT()
+struct FGatherTextExcludePath
+{
+	GENERATED_USTRUCT_BODY()
+
+	UPROPERTY(config, EditAnywhere, Category="Pattern")
+	FString Pattern;
+
+	bool Validate(FText& OutError) const;
+};
+
+USTRUCT()
+struct FGatherTextFileExtension
+{
+	GENERATED_USTRUCT_BODY()
+
+	UPROPERTY(config, EditAnywhere, Category="Pattern")
+	FString Pattern;
+
+	bool Validate(FText& OutError) const;
+};
+
 USTRUCT()
 struct FGatherTextFromTextFilesConfiguration
 {
 	GENERATED_USTRUCT_BODY()
 
-	/* The paths of directories to be searched for text files, specified relative to the project's root, which may be parsed for text to gather. */
+	static const TArray<FGatherTextFileExtension> DefaultTextFileExtensions;
+
+	FGatherTextFromTextFilesConfiguration()
+		: IsEnabled(true)
+		, FileExtensions(DefaultTextFileExtensions)
+	{
+	}
+
+	/* If enabled, text from text files will be gathered according to this configuration. */
+	UPROPERTY(config, EditAnywhere, Category = "Execution")
+	bool IsEnabled;
+
+	/* The paths of directories to be searched recursively for text files, specified relative to the project's root, which may be parsed for text to gather. */
 	UPROPERTY(config, EditAnywhere, Category = "Filter")
-	TArray<FString> SearchDirectories;
+	TArray<FGatherTextSearchDirectory> SearchDirectories;
 
 	/* Text files whose paths match these wildcard patterns will be excluded from gathering. */
 	UPROPERTY(config, EditAnywhere, Category = "Filter")
-	TArray<FString> ExcludePathWildcards;
+	TArray<FGatherTextExcludePath> ExcludePathWildcards;
 
 	/* Text files whose names match these wildcard patterns may be parsed for text to gather. */
 	UPROPERTY(config, EditAnywhere, Category = "Filter")
-	TArray<FString> FileExtensions;
+	TArray<FGatherTextFileExtension> FileExtensions;
+
+	bool Validate(const FString& RootDirectory, FText& OutError) const;
 };
 
+
 USTRUCT()
 struct FGatherTextFromPackagesConfiguration
 {
 	GENERATED_USTRUCT_BODY()
 
+	static const TArray<FGatherTextFileExtension> DefaultPackageFileExtensions;
+
+	FGatherTextFromPackagesConfiguration()
+		: IsEnabled(true)
+		, FileExtensions(DefaultPackageFileExtensions)
+	{
+	}
+
+	/* If enabled, text from packages will be gathered according to this configuration. */
+	UPROPERTY(config, EditAnywhere, Category = "Execution")
+	bool IsEnabled;
+
 	/* Packages whose paths match these wildcard patterns, specified relative to the project's root, may be processed for gathering. */
 	UPROPERTY(config, EditAnywhere, Category = "Filter")
-	TArray<FString> IncludePathWildcards;
+	TArray<FGatherTextIncludePath> IncludePathWildcards;
 
 	/* Packages whose paths match these wildcard patterns will be excluded from gathering. */
 	UPROPERTY(config, EditAnywhere, Category = "Filter")
-	TArray<FString> ExcludePathWildcards;
+	TArray<FGatherTextExcludePath> ExcludePathWildcards;
 
 	/* Packages whose names match these wildcard patterns may be processed for text to gather. */
 	UPROPERTY(config, EditAnywhere, Category = "Filter")
-	TArray<FString> FileExtensions;
+	TArray<FGatherTextFileExtension> FileExtensions;
+
+	bool Validate(const FString& RootDirectory, FText& OutError) const;
+};
+
+USTRUCT()
+struct FMetaDataTextKeyPattern
+{
+	GENERATED_USTRUCT_BODY()
+
+	UPROPERTY(config, EditAnywhere, Category="Pattern")
+	FString Pattern;
+
+	static const TArray<FString> PossiblePlaceHolders;
+
+	bool Validate(FText& OutError) const;
+};
+
+USTRUCT()
+struct FMetaDataKeyName
+{
+	GENERATED_USTRUCT_BODY()
+
+	UPROPERTY(config, EditAnywhere, Category="Name")
+	FString Name;
+
+	bool Validate(FText& OutError) const;
 };
 
 USTRUCT()
@@ -56,7 +153,7 @@ struct FMetaDataKeyGatherSpecification
 
 	/*  The metadata key for which values will be gathered as text. */
 	UPROPERTY(config, EditAnywhere, Category = "Input")
-	FString MetaDataKey;
+	FMetaDataKeyName MetaDataKey;
 
 	/* The localization namespace in which the gathered text will be output. */
 	UPROPERTY(config, EditAnywhere, Category = "Output")
@@ -67,7 +164,9 @@ struct FMetaDataKeyGatherSpecification
 	{FieldPath} - The fully qualified name of the object upon which the metadata resides.
 	{MetaDataValue} - The value associated with the metadata key. */
 	UPROPERTY(config, EditAnywhere, Category = "Output")
-	FString TextKeyPattern;
+	FMetaDataTextKeyPattern TextKeyPattern;
+
+	bool Validate(FText& OutError) const;
 };
 
 USTRUCT()
@@ -75,17 +174,28 @@ struct FGatherTextFromMetaDataConfiguration
 {
 	GENERATED_USTRUCT_BODY()
 
+	FGatherTextFromMetaDataConfiguration()
+		: IsEnabled(false)
+	{
+	}
+
+	/* If enabled, metadata will be gathered according to this configuration. */
+	UPROPERTY(config, EditAnywhere, Category = "Execution")
+	bool IsEnabled;
+
 	/* Metadata from source files whose paths match these wildcard patterns, specified relative to the project's root, may be processed for gathering. */
 	UPROPERTY(config, EditAnywhere, Category = "Filter")
-	TArray<FString> IncludePathWildcards;
+	TArray<FGatherTextIncludePath> IncludePathWildcards;
 
 	/* Metadata from source files whose paths match these wildcard patterns will be excluded from gathering. */
 	UPROPERTY(config, EditAnywhere, Category = "Filter")
-	TArray<FString> ExcludePathWildcards;
+	TArray<FGatherTextExcludePath> ExcludePathWildcards;
 
 	/* Specifications for how to gather text from specific metadata keys. */
 	UPROPERTY(config, EditAnywhere, Category = "MetaData")
 	TArray<FMetaDataKeyGatherSpecification> KeySpecifications;
+
+	bool Validate(const FString& RootDirectory, FText& OutError) const;
 };
 
 USTRUCT()
diff --git a/Engine/Source/Editor/UnrealEd/Private/Commandlets/GatherTextCommandletBase.cpp b/Engine/Source/Editor/UnrealEd/Private/Commandlets/GatherTextCommandletBase.cpp
--- a/Engine/Source/Editor/UnrealEd/Private/Commandlets/GatherTextCommandletBase.cpp
+++ b/Engine/Source/Editor/UnrealEd/Private/Commandlets/GatherTextCommandletBase.cpp
@@ -218,15 +218,11 @@ int32 UGatherTextCommandletBase::GetPathArrayFromConfig( const TCHAR* Section, c
 	{
 		if (FPaths::IsRelative(OutArr[i]))
 		{
-			if (!FPaths::GameDir().IsEmpty())
-			{
-				OutArr[i] = FPaths::Combine( *( FPaths::GameDir() ), *OutArr[i] );
-			}
-			else
-			{
-				OutArr[i] = FPaths::Combine( *( FPaths::EngineDir() ), *OutArr[i] );
-			}
+			const FString ProjectBasePath = FPaths::GameDir().IsEmpty() ? FPaths::EngineDir() : FPaths::GameDir();
+			OutArr[i] = FPaths::Combine( *ProjectBasePath, *OutArr[i] );
+			OutArr[i] = FPaths::ConvertRelativePathToFull(OutArr[i]);
 		}
+		FPaths::CollapseRelativeDirectories(OutArr[i]);
 	}
 	return count;
 }
diff --git a/Engine/Source/Editor/UnrealEd/Private/Commandlets/GatherTextFromAssetsCommandlet.cpp b/Engine/Source/Editor/UnrealEd/Private/Commandlets/GatherTextFromAssetsCommandlet.cpp
--- a/Engine/Source/Editor/UnrealEd/Private/Commandlets/GatherTextFromAssetsCommandlet.cpp
+++ b/Engine/Source/Editor/UnrealEd/Private/Commandlets/GatherTextFromAssetsCommandlet.cpp
@@ -719,52 +719,61 @@ int32 UGatherTextFromAssetsCommandlet::Main(const FString& Params)
 		FModuleManager::Get().LoadModule(*ModuleName);
 	}
 
-	//Include paths
-	TArray<FString> IncludePaths;
-	GetPathArrayFromConfig(*SectionName, TEXT("IncludePaths"), IncludePaths, GatherTextConfigPath);
+	// IncludePathFilters
+	TArray<FString> IncludePathFilters;
+	GetPathArrayFromConfig(*SectionName, TEXT("IncludePathFilters"), IncludePathFilters, GatherTextConfigPath);
 
-	if (IncludePaths.Num() == 0)
+	// IncludePaths (DEPRECATED)
 	{
-		UE_LOG(LogGatherTextFromAssetsCommandlet, Error, TEXT("No include paths in section %s"), *SectionName);
-		return -1;
+		TArray<FString> IncludePaths;
+		GetPathArrayFromConfig(*SectionName, TEXT("IncludePaths"), IncludePaths, GatherTextConfigPath);
+		if (IncludePaths.Num())
+		{
+			IncludePathFilters.Append(IncludePaths);
+			UE_LOG(LogGatherTextFromAssetsCommandlet, Warning, TEXT("IncludePaths detected in section %s. IncludePaths is deprecated, please use IncludePathFilters."), *SectionName);
+		}
 	}
 
-	for (FString& IncludePath : IncludePaths)
+	if (IncludePathFilters.Num() == 0)
 	{
-		FPaths::NormalizeDirectoryName(IncludePath);
+		UE_LOG(LogGatherTextFromAssetsCommandlet, Error, TEXT("No include path filters in section %s."), *SectionName);
+		return -1;
+	}
 
-		IncludePath = FPaths::ConvertRelativePathToFull(IncludePath);
+	// ExcludePathFilters
+	TArray<FString> ExcludePathFilters;
+	GetPathArrayFromConfig(*SectionName, TEXT("ExcludePathFilters"), ExcludePathFilters, GatherTextConfigPath);
 
-		// All paths must ends with "/*"
-		if ( !IncludePath.EndsWith(TEXT("/*")) )
+	// ExcludePaths (DEPRECATED)
+	{
+		TArray<FString> ExcludePaths;
+		GetPathArrayFromConfig(*SectionName, TEXT("ExcludePaths"), ExcludePaths, GatherTextConfigPath);
+		if (ExcludePaths.Num())
 		{
-			// If it ends in a slash, add the star.
-			if ( IncludePath.EndsWith(TEXT("/")) )
-			{
-				IncludePath.AppendChar(TEXT('*'));
-			}
-			// If it doesn't end in a slash or slash star, just add slash star.
-			else
-			{
-				IncludePath.Append(TEXT("/*"));
-			}
+			ExcludePathFilters.Append(ExcludePaths);
+			UE_LOG(LogGatherTextFromAssetsCommandlet, Warning, TEXT("ExcludePaths detected in section %s. ExcludePaths is deprecated, please use ExcludePathFilters."), *SectionName);
 		}
 	}
 
-	//Exclude paths
-	TArray<FString> ExcludePaths;
-	GetStringArrayFromConfig(*SectionName, TEXT("ExcludePaths"), ExcludePaths, GatherTextConfigPath);
+	// PackageNameFilters
+	TArray<FString> PackageFileNameFilters;
+	GetStringArrayFromConfig(*SectionName, TEXT("PackageFileNameFilters"), PackageFileNameFilters, GatherTextConfigPath);
 
-	//package extensions
-	TArray<FString> PackageExts;
-	GetStringArrayFromConfig(*SectionName, TEXT("PackageExtensions"), PackageExts, GatherTextConfigPath);
-
-	if (PackageExts.Num() == 0)
+	// PackageExtensions (DEPRECATED)
 	{
-		UE_LOG(LogGatherTextFromAssetsCommandlet, Warning, TEXT("No package extensions specified in section %s, using defaults"), *SectionName);
+		TArray<FString> PackageExtensions;
+		GetStringArrayFromConfig(*SectionName, TEXT("PackageExtensions"), PackageExtensions, GatherTextConfigPath);
+		if (PackageExtensions.Num())
+		{
+			PackageFileNameFilters.Append(PackageExtensions);
+			UE_LOG(LogGatherTextFromAssetsCommandlet, Warning, TEXT("PackageExtensions detected in section %s. PackageExtensions is deprecated, please use PackageFileNameFilters."), *SectionName);
+		}
+	}
 
-		PackageExts.Add(FString("*") + FPackageName::GetAssetPackageExtension());
-		PackageExts.Add(FString("*") + FPackageName::GetMapPackageExtension());
+	if (PackageFileNameFilters.Num() == 0)
+	{
+		UE_LOG(LogGatherTextFromAssetsCommandlet, Error, TEXT("No package file name filters in section %s."), *SectionName);
+		return -1;
 	}
 
 	//asset class exclude
@@ -831,19 +840,19 @@ int32 UGatherTextFromAssetsCommandlet::Main(const FString& Params)
 	//Fill the list of packages to work from.
 	uint8 PackageFilter = NORMALIZE_DefaultFlags;
 	TArray<FString> Unused;	
-	for ( int32 PackageFilenameWildcardIdx = 0; PackageFilenameWildcardIdx < PackageExts.Num(); PackageFilenameWildcardIdx++ )
+	for ( int32 PackageFilenameWildcardIdx = 0; PackageFilenameWildcardIdx < PackageFileNameFilters.Num(); PackageFilenameWildcardIdx++ )
 	{
-		const bool IsAssetPackage = PackageExts[PackageFilenameWildcardIdx] == ( FString( TEXT("*") )+ FPackageName::GetAssetPackageExtension() );
+		const bool IsAssetPackage = PackageFileNameFilters[PackageFilenameWildcardIdx] == ( FString( TEXT("*") )+ FPackageName::GetAssetPackageExtension() );
 
 		TArray<FString> PackageFiles;
-		if ( !NormalizePackageNames( Unused, PackageFiles, PackageExts[PackageFilenameWildcardIdx], PackageFilter) )
+		if ( !NormalizePackageNames( Unused, PackageFiles, PackageFileNameFilters[PackageFilenameWildcardIdx], PackageFilter) )
 		{
-			UE_LOG(LogGatherTextFromAssetsCommandlet, Display, TEXT("No packages found with extension %i: '%s'"), PackageFilenameWildcardIdx, *PackageExts[PackageFilenameWildcardIdx]);
+			UE_LOG(LogGatherTextFromAssetsCommandlet, Display, TEXT("No packages found with extension %i: '%s'"), PackageFilenameWildcardIdx, *PackageFileNameFilters[PackageFilenameWildcardIdx]);
 			continue;
 		}
 		else
 		{
-			UE_LOG(LogGatherTextFromAssetsCommandlet, Display, TEXT("Found %i packages with extension %i: '%s'"), PackageFiles.Num(), PackageFilenameWildcardIdx, *PackageExts[PackageFilenameWildcardIdx]);
+			UE_LOG(LogGatherTextFromAssetsCommandlet, Display, TEXT("Found %i packages with extension %i: '%s'"), PackageFiles.Num(), PackageFilenameWildcardIdx, *PackageFileNameFilters[PackageFilenameWildcardIdx]);
 		}
 
 		//Run through all the files found and add any that pass the include, exclude and filter constraints to OrderedPackageFilesToLoad
@@ -853,7 +862,7 @@ int32 UGatherTextFromAssetsCommandlet::Main(const FString& Params)
 
 			bool bExclude = false;
 			//Ensure it matches the include paths if there are some.
-			for (FString& IncludePath : IncludePaths)
+			for (FString& IncludePath : IncludePathFilters)
 			{
 				bExclude = true;
 				if( PackageFile.MatchesWildcard(IncludePath) )
@@ -869,9 +878,9 @@ int32 UGatherTextFromAssetsCommandlet::Main(const FString& Params)
 			}
 
 			//Ensure it does not match the exclude paths if there are some.
-			for( int32 ExcludePathIdx=0; !bExclude && ExcludePathIdx<ExcludePaths.Num() ; ++ExcludePathIdx )
+			for (const FString& ExcludePath : ExcludePathFilters)
 			{
-				if( PackageFile.MatchesWildcard(ExcludePaths[ExcludePathIdx]) )
+				if (PackageFile.MatchesWildcard(ExcludePath))
 				{
 					bExclude = true;
 					PackageFilesInExcludePath.Add(PackageFile);
diff --git a/Engine/Source/Editor/UnrealEd/Private/Commandlets/GatherTextFromMetadataCommandlet.cpp b/Engine/Source/Editor/UnrealEd/Private/Commandlets/GatherTextFromMetadataCommandlet.cpp
--- a/Engine/Source/Editor/UnrealEd/Private/Commandlets/GatherTextFromMetadataCommandlet.cpp
+++ b/Engine/Source/Editor/UnrealEd/Private/Commandlets/GatherTextFromMetadataCommandlet.cpp
@@ -58,42 +58,42 @@ int32 UGatherTextFromMetaDataCommandlet::Main( const FString& Params )
 		FModuleManager::Get().LoadModule(*ModuleName);
 	}
 
-	//Include paths
-	TArray<FString> IncludePaths;
-	GetPathArrayFromConfig(*SectionName, TEXT("IncludePaths"), IncludePaths, GatherTextConfigPath);
+	// IncludePathFilters
+	TArray<FString> IncludePathFilters;
+	GetPathArrayFromConfig(*SectionName, TEXT("IncludePathFilters"), IncludePathFilters, GatherTextConfigPath);
 
-	if (IncludePaths.Num() == 0)
+	// IncludePaths (DEPRECATED)
 	{
-		UE_LOG(LogGatherTextFromMetaDataCommandlet, Error, TEXT("No include paths in section %s"), *SectionName);
-		return -1;
+		TArray<FString> IncludePaths;
+		GetPathArrayFromConfig(*SectionName, TEXT("IncludePaths"), IncludePaths, GatherTextConfigPath);
+		if (IncludePaths.Num())
+		{
+			IncludePathFilters.Append(IncludePaths);
+			UE_LOG(LogGatherTextFromMetaDataCommandlet, Warning, TEXT("IncludePaths detected in section %s. IncludePaths is deprecated, please use IncludePathFilters."), *SectionName);
+		}
 	}
 
-	for (FString& IncludePath : IncludePaths)
+	if (IncludePathFilters.Num() == 0)
 	{
-		FPaths::NormalizeDirectoryName(IncludePath);
+		UE_LOG(LogGatherTextFromMetaDataCommandlet, Error, TEXT("No include path filters in section %s."), *SectionName);
+		return -1;
+	}
 
-		IncludePath = FPaths::ConvertRelativePathToFull(IncludePath);
+	// ExcludePathFilters
+	TArray<FString> ExcludePathFilters;
+	GetPathArrayFromConfig(*SectionName, TEXT("ExcludePathFilters"), ExcludePathFilters, GatherTextConfigPath);
 
-		// All paths must ends with "/*"
-		if ( !IncludePath.EndsWith(TEXT("/*")) )
+	// ExcludePaths (DEPRECATED)
+	{
+		TArray<FString> ExcludePaths;
+		GetPathArrayFromConfig(*SectionName, TEXT("ExcludePaths"), ExcludePaths, GatherTextConfigPath);
+		if (ExcludePaths.Num())
 		{
-			// If it ends in a slash, add the star.
-			if ( IncludePath.EndsWith(TEXT("/")) )
-			{
-				IncludePath.AppendChar(TEXT('*'));
-			}
-			// If it doesn't end in a slash or slash star, just add slash star.
-			else
-			{
-				IncludePath.Append(TEXT("/*"));
-			}
+			ExcludePathFilters.Append(ExcludePaths);
+			UE_LOG(LogGatherTextFromMetaDataCommandlet, Warning, TEXT("ExcludePaths detected in section %s. ExcludePaths is deprecated, please use ExcludePathFilters."), *SectionName);
 		}
 	}
 
-	//Exclude paths
-	TArray<FString> ExcludePaths;
-	GetStringArrayFromConfig(*SectionName, TEXT("ExcludePaths"), ExcludePaths, GatherTextConfigPath);
-
 	FGatherParameters Arguments;
 	GetStringArrayFromConfig(*SectionName, TEXT("InputKeys"), Arguments.InputKeys, GatherTextConfigPath);
 	GetStringArrayFromConfig(*SectionName, TEXT("OutputNamespaces"), Arguments.OutputNamespaces, GatherTextConfigPath);
@@ -105,7 +105,7 @@ int32 UGatherTextFromMetaDataCommandlet::Main( const FString& Params )
 	}
 
 	// Execute gather.
-	GatherTextFromUObjects(IncludePaths, ExcludePaths, Arguments);
+	GatherTextFromUObjects(IncludePathFilters, ExcludePathFilters, Arguments);
 
 	// Add any manifest dependencies if they were provided
 	TArray<FString> ManifestDependenciesList;
diff --git a/Engine/Source/Editor/UnrealEd/Private/Commandlets/GatherTextFromSourceCommandlet.cpp b/Engine/Source/Editor/UnrealEd/Private/Commandlets/GatherTextFromSourceCommandlet.cpp
--- a/Engine/Source/Editor/UnrealEd/Private/Commandlets/GatherTextFromSourceCommandlet.cpp
+++ b/Engine/Source/Editor/UnrealEd/Private/Commandlets/GatherTextFromSourceCommandlet.cpp
@@ -58,26 +58,47 @@ int32 UGatherTextFromSourceCommandlet::Main( const FString& Params )
 		return -1;
 	}
 
+	// SearchDirectoryPaths
+	TArray<FString> SearchDirectoryPaths;
+	GetPathArrayFromConfig(*SectionName, TEXT("SearchDirectoryPaths"), SearchDirectoryPaths, GatherTextConfigPath);
 
-	//Include paths
-	TArray<FString> IncludePaths;
-	GetStringArrayFromConfig(*SectionName, TEXT("IncludePaths"), IncludePaths, GatherTextConfigPath);
+	// IncludePaths (DEPRECATED)
+	{
+		TArray<FString> IncludePaths;
+		GetPathArrayFromConfig(*SectionName, TEXT("IncludePaths"), IncludePaths, GatherTextConfigPath);
+		if (IncludePaths.Num())
+		{
+			SearchDirectoryPaths.Append(IncludePaths);
+			UE_LOG(LogGatherTextFromSourceCommandlet, Warning, TEXT("IncludePaths detected in section %s. IncludePaths is deprecated, please use SearchDirectoryPaths."), *SectionName);
+		}
+	}
 
-	if (IncludePaths.Num() == 0)
+	if (SearchDirectoryPaths.Num() == 0)
 	{
-		UE_LOG(LogGatherTextFromSourceCommandlet, Error, TEXT("No include paths in section %s"), *SectionName);
+		UE_LOG(LogGatherTextFromSourceCommandlet, Error, TEXT("No search directory paths in section %s."), *SectionName);
 		return -1;
 	}
 
-	//Exclude paths
+	// Exclude paths
 	TArray<FString> ExcludePaths;
-	GetStringArrayFromConfig(*SectionName, TEXT("ExcludePaths"), ExcludePaths, GatherTextConfigPath);
+	GetPathArrayFromConfig(*SectionName, TEXT("ExcludePaths"), ExcludePaths, GatherTextConfigPath);
 
-	// Check the config section for source filters. e.g. *.cpp
-	TArray<FString> SourceFileSearchFilters;
-	GetStringArrayFromConfig(*SectionName, TEXT("SourceFileSearchFilters"), SourceFileSearchFilters, GatherTextConfigPath);
+	// FileNameFilters
+	TArray<FString> FileNameFilters;
+	GetStringArrayFromConfig(*SectionName, TEXT("FileNameFilters"), FileNameFilters, GatherTextConfigPath);
+
+	// SourceFileSearchFilters (DEPRECATED)
+	{
+		TArray<FString> SourceFileSearchFilters;
+		GetPathArrayFromConfig(*SectionName, TEXT("SourceFileSearchFilters"), SourceFileSearchFilters, GatherTextConfigPath);
+		if (SourceFileSearchFilters.Num())
+		{
+			FileNameFilters.Append(SourceFileSearchFilters);
+			UE_LOG(LogGatherTextFromSourceCommandlet, Warning, TEXT("SourceFileSearchFilters detected in section %s. SourceFileSearchFilters is deprecated, please use FileNameFilters."), *SectionName);
+		}
+	}
 
-	if (SourceFileSearchFilters.Num() == 0)
+	if (FileNameFilters.Num() == 0)
 	{
 		UE_LOG(LogGatherTextFromSourceCommandlet, Error, TEXT("No source filters in section %s"), *SectionName);
 		return -1;
@@ -85,38 +106,29 @@ int32 UGatherTextFromSourceCommandlet::Main( const FString& Params )
 
 	//Ensure all filters are unique.
 	TArray<FString> UniqueSourceFileSearchFilters;
-	for (int32 Idx=0; Idx<SourceFileSearchFilters.Num(); Idx++)
+	for (const FString& SourceFileSearchFilter : FileNameFilters)
 	{
-		UniqueSourceFileSearchFilters.AddUnique(SourceFileSearchFilters[Idx]);
+		UniqueSourceFileSearchFilters.AddUnique(SourceFileSearchFilter);
 	}
 
 	// Search in the root folder for each of the wildcard filters specified and build a list of files
 	TArray<FString> AllFoundFiles;
 
-	for (FString& IncludePath : IncludePaths)
+	for (FString& SearchDirectoryPath : SearchDirectoryPaths)
 	{
-		FString ProjectBasePath;
-		if (!FPaths::GameDir().IsEmpty())
-		{
-			ProjectBasePath = FPaths::GameDir();
-		}
-		else
-		{
-			ProjectBasePath = FPaths::EngineDir();
-		}
-
-		if(FPaths::IsRelative(IncludePath))
-		{
-			IncludePath = FPaths::Combine( *ProjectBasePath, *IncludePath );
-			IncludePath = FPaths::ConvertRelativePathToFull(IncludePath);
-			FPaths::CollapseRelativeDirectories(IncludePath);
-		}
-
-		for (int32 SourceFileSearchIdx=0; SourceFileSearchIdx < UniqueSourceFileSearchFilters.Num(); SourceFileSearchIdx++)
+		for (const FString& UniqueSourceFileSearchFilter : UniqueSourceFileSearchFilters)
 		{
 			TArray<FString> RootSourceFiles;
 
-			IFileManager::Get().FindFilesRecursive(RootSourceFiles, *IncludePath, *UniqueSourceFileSearchFilters[SourceFileSearchIdx], true, false,false);
+			IFileManager::Get().FindFilesRecursive(RootSourceFiles, *SearchDirectoryPath, *UniqueSourceFileSearchFilter, true, false,false);
+
+			for (FString& RootSourceFile : RootSourceFiles)
+			{
+				if (FPaths::IsRelative(RootSourceFile))
+				{
+					RootSourceFile = FPaths::ConvertRelativePathToFull(RootSourceFile);
+				}
+			}
 
 			AllFoundFiles.Append(RootSourceFiles);
 		}
@@ -125,18 +137,18 @@ int32 UGatherTextFromSourceCommandlet::Main( const FString& Params )
 	TArray<FString> FilesToProcess;
 	TArray<FString> RemovedList;
 
-	//Run through all the files found and add any that pass the include, exclude and filter constraints to PackageFilesToProcess
-	for( int32 FileIdx=0; FileIdx < AllFoundFiles.Num(); ++FileIdx )
+	//Run through all the files found and add any that pass the exclude and filter constraints to PackageFilesToProcess
+	for (const FString& FoundFile : AllFoundFiles)
 	{
 		bool bExclude = false;
 
 		//Ensure it does not match the exclude paths if there are some.
-		for( int32 ExcludePathIdx=0; ExcludePathIdx < ExcludePaths.Num() ; ++ExcludePathIdx )
+		for (FString& ExcludePath : ExcludePaths)
 		{
-			if( AllFoundFiles[FileIdx].MatchesWildcard( ExcludePaths[ExcludePathIdx] ) )
+			if (FoundFile.MatchesWildcard(ExcludePath) )
 			{
 				bExclude = true;
-				RemovedList.Add( AllFoundFiles[FileIdx] );
+				RemovedList.Add(FoundFile);
 				break;
 			}
 		}
@@ -144,7 +156,7 @@ int32 UGatherTextFromSourceCommandlet::Main( const FString& Params )
 		//If we haven't failed any checks, add it to the array of files to process.
 		if( !bExclude )
 		{
-			FilesToProcess.Add( AllFoundFiles[FileIdx] );
+			FilesToProcess.Add(FoundFile);
 		}
 	}
 	
@@ -152,9 +164,9 @@ int32 UGatherTextFromSourceCommandlet::Main( const FString& Params )
 	if( FilesToProcess.Num() == 0 )
 	{
 		FString SpecifiedDirectoriesString;
-		for (FString& IncludePath : IncludePaths)
+		for (FString& SearchDirectoryPath : SearchDirectoryPaths)
 		{
-			SpecifiedDirectoriesString.Append(FString(SpecifiedDirectoriesString.IsEmpty() ? TEXT("") : TEXT("\n")) + FString::Printf(TEXT("+ %s"), *IncludePath));
+			SpecifiedDirectoriesString.Append(FString(SpecifiedDirectoriesString.IsEmpty() ? TEXT("") : TEXT("\n")) + FString::Printf(TEXT("+ %s"), *SearchDirectoryPath));
 		}
 		for (FString& ExcludePath : ExcludePaths)
 		{
