Fixed horizontal alignment on meta data text key pattern in the localization dashboard UI.

diff --git a/Engine/Source/Editor/LocalizationDashboard/Private/GatherTextDetailCustomizations.cpp b/Engine/Source/Editor/LocalizationDashboard/Private/GatherTextDetailCustomizations.cpp
--- a/Engine/Source/Editor/LocalizationDashboard/Private/GatherTextDetailCustomizations.cpp
+++ b/Engine/Source/Editor/LocalizationDashboard/Private/GatherTextDetailCustomizations.cpp
@@ -775,6 +775,7 @@ void FMetaDataTextKeyPatternStructCustomization::CustomizeStructHeader( TSharedR
 			]
 		]
 		.ValueContent()
+		.HAlign(HAlign_Fill)
 		[
 			SNew(SMetaDataTextKeyPatternWidget, PatternPropertyHandle.ToSharedRef())
 		];
