Moved specification of the native culture of a localization target over to the table of supported cultures.

diff --git a/Engine/Source/Editor/LocalizationDashboard/Private/LocalizationConfigurationScript.cpp b/Engine/Source/Editor/LocalizationDashboard/Private/LocalizationConfigurationScript.cpp
--- a/Engine/Source/Editor/LocalizationDashboard/Private/LocalizationConfigurationScript.cpp
+++ b/Engine/Source/Editor/LocalizationDashboard/Private/LocalizationConfigurationScript.cpp
@@ -130,8 +130,7 @@ namespace LocalizationConfigurationScript
 			ConfigSection.Add( TEXT("ManifestName"), FPaths::GetCleanFilename(GetManifestPath(Target)) );
 			ConfigSection.Add( TEXT("ArchiveName"), FPaths::GetCleanFilename(GetArchivePath(Target, FString())) );
 
-			ConfigSection.Add( TEXT("SourceCulture"), Target->Settings.NativeCultureStatistics.CultureName );
-			ConfigSection.Add( TEXT("CulturesToGenerate"), Target->Settings.NativeCultureStatistics.CultureName );
+			ConfigSection.Add( TEXT("SourceCulture"), Target->Settings.SupportedCulturesStatistics[Target->Settings.NativeCultureIndex].CultureName );
 			for (const FCultureStatistics& CultureStatistics : Target->Settings.SupportedCulturesStatistics)
 			{
 				ConfigSection.Add( TEXT("CulturesToGenerate"), CultureStatistics.CultureName );
@@ -406,28 +405,21 @@ namespace LocalizationConfigurationScript
 			}
 			ConfigSection.Add( TEXT("DestinationPath"), DestinationPath );
 
-			TArray<const FCultureStatistics*> AllCultureStatistics;
-			AllCultureStatistics.Add(&Target->Settings.NativeCultureStatistics);
-			for (const FCultureStatistics& SupportedCultureStatistics : Target->Settings.SupportedCulturesStatistics)
-			{
-				AllCultureStatistics.Add(&SupportedCultureStatistics);
-			}
-
 			const auto& AddCultureToGenerate = [&](const int32 Index)
 			{
-				ConfigSection.Add( TEXT("CulturesToGenerate"), AllCultureStatistics[Index]->CultureName );
+				ConfigSection.Add( TEXT("CulturesToGenerate"), Target->Settings.SupportedCulturesStatistics[Index].CultureName );
 			};
 
 			// Export for a specific culture.
 			if (CultureName.IsSet())
 			{
-				const int32 CultureIndex = AllCultureStatistics.IndexOfByPredicate([CultureName](const FCultureStatistics* Culture) { return Culture->CultureName == CultureName.GetValue(); });
+				const int32 CultureIndex = Target->Settings.SupportedCulturesStatistics.IndexOfByPredicate([CultureName](const FCultureStatistics& Culture) { return Culture.CultureName == CultureName.GetValue(); });
 				AddCultureToGenerate(CultureIndex);
 			}
 			// Export for all cultures.
 			else
 			{
-				for (int32 CultureIndex = 0; CultureIndex < AllCultureStatistics.Num(); ++CultureIndex)
+				for (int32 CultureIndex = 0; CultureIndex < Target->Settings.SupportedCulturesStatistics.Num(); ++CultureIndex)
 				{
 					AddCultureToGenerate(CultureIndex);
 				}
@@ -537,16 +529,9 @@ namespace LocalizationConfigurationScript
 			ConfigSection.Add( TEXT("ManifestName"), FString::Printf( TEXT("%s.%s"), *Target->Settings.Name, TEXT("manifest") ) );
 			ConfigSection.Add( TEXT("ResourceName"), FString::Printf( TEXT("%s.%s"), *Target->Settings.Name, TEXT("locres") ) );
 
-			TArray<const FCultureStatistics*> AllCultureStatistics;
-			AllCultureStatistics.Add(&Target->Settings.NativeCultureStatistics);
-			for (const FCultureStatistics& SupportedCultureStatistics : Target->Settings.SupportedCulturesStatistics)
-			{
-				AllCultureStatistics.Add(&SupportedCultureStatistics);
-			}
-
-			for (const FCultureStatistics* CultureStatistics : AllCultureStatistics)
+			for (const FCultureStatistics& CultureStatistics : Target->Settings.SupportedCulturesStatistics)
 			{
-				ConfigSection.Add( TEXT("CulturesToGenerate"), CultureStatistics->CultureName );
+				ConfigSection.Add( TEXT("CulturesToGenerate"), CultureStatistics.CultureName );
 			}
 		}
 
diff --git a/Engine/Source/Editor/LocalizationDashboard/Private/LocalizationTargetDetailCustomization.cpp b/Engine/Source/Editor/LocalizationDashboard/Private/LocalizationTargetDetailCustomization.cpp
--- a/Engine/Source/Editor/LocalizationDashboard/Private/LocalizationTargetDetailCustomization.cpp
+++ b/Engine/Source/Editor/LocalizationDashboard/Private/LocalizationTargetDetailCustomization.cpp
@@ -1,4 +1,4 @@
-// Copyright 1998-2015 Epic Games, Inc. All Rights Reserved.
+﻿// Copyright 1998-2015 Epic Games, Inc. All Rights Reserved.
 
 #include "LocalizationDashboardPrivatePCH.h"
 #include "LocalizationTargetDetailCustomization.h"
@@ -169,48 +169,9 @@ void FLocalizationTargetDetailCustomization::CustomizeDetails(IDetailLayoutBuild
 		RebuildTargetDependenciesBox();
 	});
 
-	PropertyCustomizationMap.Add(GET_MEMBER_NAME_CHECKED(FLocalizationTargetSettings, NativeCultureStatistics), [&](const TSharedRef<IPropertyHandle>& MemberPropertyHandle, IDetailCategoryBuilder& DetailCategoryBuilder)
+	PropertyCustomizationMap.Add(GET_MEMBER_NAME_CHECKED(FLocalizationTargetSettings, NativeCultureIndex), [&](const TSharedRef<IPropertyHandle>& MemberPropertyHandle, IDetailCategoryBuilder& DetailCategoryBuilder)
 	{
-		NativeCultureStatisticsPropertyHandle = MemberPropertyHandle;
-		NativeCultureNamePropertyHandle = NativeCultureStatisticsPropertyHandle->GetChildHandle(GET_MEMBER_NAME_CHECKED(FCultureStatistics, CultureName));
-
-		if (NativeCultureNamePropertyHandle.IsValid() && NativeCultureNamePropertyHandle->IsValidHandle())
-		{
-			NativeCultureNamePropertyHandle->MarkHiddenByCustomization();
-			FDetailWidgetRow& NativeCultureRow = DetailCategoryBuilder.AddCustomRow( NativeCultureNamePropertyHandle->GetPropertyDisplayName() );
-
-			FString NativeCultureName;
-			NativeCultureNamePropertyHandle->GetValue(NativeCultureName);
-			FCulturePtr NativeCulture = FInternationalization::Get().GetCulture(NativeCultureName);
-
-			NativeCultureRow.NameContent()
-				[
-					NativeCultureNamePropertyHandle->CreatePropertyNameWidget(LOCTEXT("NativeCultureNameLabel", "Native Culture"))
-				];
-
-			NativeCultureRow.ValueContent()
-				[
-					SAssignNew(NativeCultureComboButton, SComboButton)
-					.ButtonContent()
-					[
-						SNew(STextBlock)
-						.Text(this, &FLocalizationTargetDetailCustomization::GetNativeCultureDisplayName)
-						.ToolTipText(this, &FLocalizationTargetDetailCustomization::GetNativeCultureName)
-					]
-					.HasDownArrow(true)
-						.MenuContent()
-						[
-							SNew(SBox)
-							.MaxDesiredHeight(400.0f)
-							.MaxDesiredWidth(300.0f)
-							[
-								SNew(SCulturePicker)
-								.OnSelectionChanged(this, &FLocalizationTargetDetailCustomization::OnNativeCultureSelected)
-								.InitialSelection(NativeCulture)
-							]
-						]
-				];
-		}
+		NativeCultureIndexPropertyHandle = MemberPropertyHandle;
 	});
 
 	PropertyCustomizationMap.Add(GET_MEMBER_NAME_CHECKED(FLocalizationTargetSettings, SupportedCulturesStatistics), [&](const TSharedRef<IPropertyHandle>& MemberPropertyHandle, IDetailCategoryBuilder& DetailCategoryBuilder)
@@ -260,21 +221,30 @@ void FLocalizationTargetDetailCustomization::CustomizeDetails(IDetailLayoutBuild
 						.HeaderRow
 						(
 						SNew(SHeaderRow)
+						+SHeaderRow::Column("IsNative")
+						.DefaultLabel( NSLOCTEXT("LocalizationCulture", "IsNativeColumnLabel", "Native"))
+						.HAlignHeader(HAlign_Center)
+						.HAlignCell(HAlign_Center)
+						.VAlignCell(VAlign_Center)
+						.FillWidth(0.1f)
 						+SHeaderRow::Column("Culture")
 						.DefaultLabel( NSLOCTEXT("LocalizationCulture", "CultureColumnLabel", "Culture"))
 						.HAlignHeader(HAlign_Fill)
 						.HAlignCell(HAlign_Fill)
 						.VAlignCell(VAlign_Center)
+						.FillWidth(0.2f)
 						+SHeaderRow::Column("WordCount")
 						.DefaultLabel( NSLOCTEXT("LocalizationCulture", "WordCountColumnLabel", "Word Count"))
 						.HAlignHeader(HAlign_Center)
 						.HAlignCell(HAlign_Fill)
 						.VAlignCell(VAlign_Center)
+						.FillWidth(0.4f)
 						+SHeaderRow::Column("Actions")
 						.DefaultLabel( NSLOCTEXT("LocalizationCulture", "ActionsColumnLabel", "Actions"))
 						.HAlignHeader(HAlign_Center)
 						.HAlignCell(HAlign_Center)
 						.VAlignCell(VAlign_Center)
+						.FillWidth(0.3f)
 						)
 					]
 				+SVerticalBox::Slot()
@@ -681,53 +651,6 @@ ECheckBoxState FLocalizationTargetDetailCustomization::IsTargetDependencyChecked
 	return LocalizationTarget->Settings.TargetDependencies.Contains(OtherLocalizationTarget->Settings.Guid) ? ECheckBoxState::Checked : ECheckBoxState::Unchecked;
 }
 
-FText FLocalizationTargetDetailCustomization::GetNativeCultureName() const
-{
-	FCulturePtr NativeCulture;
-
-	if(NativeCultureNamePropertyHandle.IsValid() && NativeCultureNamePropertyHandle->IsValidHandle())
-	{
-		FString Value;
-		NativeCultureNamePropertyHandle->GetValue(Value);
-		NativeCulture = FInternationalization::Get().GetCulture(Value);
-	}
-
-	return NativeCulture.IsValid() ? FText::FromString(NativeCulture->GetName()) : FText::GetEmpty();
-}
-
-FText FLocalizationTargetDetailCustomization::GetNativeCultureDisplayName() const
-{
-	FCulturePtr NativeCulture;
-
-	if(NativeCultureNamePropertyHandle.IsValid() && NativeCultureNamePropertyHandle->IsValidHandle())
-	{
-		FString Value;
-		NativeCultureNamePropertyHandle->GetValue(Value);
-		NativeCulture = FInternationalization::Get().GetCulture(Value);
-	}
-
-	return NativeCulture.IsValid() ? FText::FromString(NativeCulture->GetDisplayName()) : FText::GetEmpty();
-}
-
-void FLocalizationTargetDetailCustomization::OnNativeCultureSelected(FCulturePtr SelectedCulture, ESelectInfo::Type SelectInfo)
-{
-	if(NativeCultureNamePropertyHandle.IsValid() && NativeCultureNamePropertyHandle->IsValidHandle())
-	{
-		NativeCultureNamePropertyHandle->SetValue(SelectedCulture->GetName());
-	}
-	if (NativeCultureComboButton.IsValid())
-	{
-		NativeCultureComboButton->SetIsOpen(false);
-	}
-
-	// TODO:	If the selected culture is an existing supported culture:
-	//			copy that supported culture's statistics to the native culture statistics
-	//			remove the supported culture's statistics
-
-	// TODO: Warn about this action invalidating all word counts and, realistically, invalidating all existing translations.
-	// Possibly offer an option to delete all archives?
-}
-
 void FLocalizationTargetDetailCustomization::Gather()
 {
 	if (LocalizationTarget.IsValid())
@@ -856,16 +779,6 @@ void FLocalizationTargetDetailCustomization::UpdateTargetFromReports()
 
 		if (TargetSettingsPropertyHandle.IsValid() && TargetSettingsPropertyHandle->IsValidHandle())
 		{
-			const TSharedPtr<IPropertyHandle> NativeCultureStatisticsPropHandle = TargetSettingsPropertyHandle->GetChildHandle(GET_MEMBER_NAME_CHECKED(FLocalizationTargetSettings, NativeCultureStatistics));
-			if (NativeCultureStatisticsPropHandle.IsValid() && NativeCultureStatisticsPropHandle->IsValidHandle())
-			{
-				const TSharedPtr<IPropertyHandle> WordCountPropertyHandle = NativeCultureStatisticsPropHandle->GetChildHandle(GET_MEMBER_NAME_CHECKED(FCultureStatistics, WordCount));
-				if (WordCountPropertyHandle.IsValid() && WordCountPropertyHandle->IsValidHandle())
-				{
-					WordCountPropertyHandles.Add(WordCountPropertyHandle);
-				}
-			}
-
 			const TSharedPtr<IPropertyHandle> SupportedCulturesStatisticsPropHandle = TargetSettingsPropertyHandle->GetChildHandle(GET_MEMBER_NAME_CHECKED(FLocalizationTargetSettings, SupportedCulturesStatistics));
 			if (SupportedCulturesStatisticsPropHandle.IsValid() && SupportedCulturesStatisticsPropHandle->IsValidHandle())
 			{
@@ -901,7 +814,6 @@ void FLocalizationTargetDetailCustomization::UpdateTargetFromReports()
 
 void FLocalizationTargetDetailCustomization::BuildListedCulturesList()
 {
-	ListedCultureStatisticProperties.AddUnique(NativeCultureStatisticsPropertyHandle);
 	const TSharedPtr<IPropertyHandleArray> SupportedCulturesStatisticsArrayPropertyHandle = SupportedCulturesStatisticsPropertyHandle->AsArray();
 	if (SupportedCulturesStatisticsArrayPropertyHandle.IsValid())
 	{
@@ -913,6 +825,20 @@ void FLocalizationTargetDetailCustomization::BuildListedCulturesList()
 			ListedCultureStatisticProperties.AddUnique(CultureStatisticsProperty);
 		}
 	}
+
+	const auto& CultureSorter = [](const TSharedPtr<IPropertyHandle>& Left, const TSharedPtr<IPropertyHandle>& Right) -> bool
+	{
+		const TSharedPtr<IPropertyHandle> LeftNameHandle = Left->GetChildHandle(GET_MEMBER_NAME_CHECKED(FCultureStatistics, CultureName));
+		const TSharedPtr<IPropertyHandle> RightNameHandle = Right->GetChildHandle(GET_MEMBER_NAME_CHECKED(FCultureStatistics, CultureName));
+
+		FString LeftName;
+		LeftNameHandle->GetValue(LeftName);
+		FString RightName;
+		RightNameHandle->GetValue(RightName);
+
+		return LeftName < RightName;
+	};
+	ListedCultureStatisticProperties.Sort(CultureSorter);
 }
 
 void FLocalizationTargetDetailCustomization::RebuildListedCulturesList()
@@ -954,18 +880,6 @@ TSharedRef<ITableRow> FLocalizationTargetDetailCustomization::OnGenerateCultureR
 
 bool FLocalizationTargetDetailCustomization::IsCultureSelectableAsSupported(FCulturePtr Culture)
 {
-	// Can't select the native culture.
-	if (NativeCultureNamePropertyHandle.IsValid() && NativeCultureNamePropertyHandle->IsValidHandle())
-	{
-		FString NativeCultureName;
-		NativeCultureNamePropertyHandle->GetValue(NativeCultureName);
-		const FCulturePtr NativeCulture = FInternationalization::Get().GetCulture(NativeCultureName);
-		if (NativeCulture == Culture)
-		{
-			return false;
-		}
-	}
-
 	auto Is = [&](const TSharedPtr<IPropertyHandle>& SupportedCultureStatisticProperty)
 	{
 		// Can't select existing supported cultures.
@@ -996,6 +910,17 @@ void FLocalizationTargetDetailCustomization::OnNewSupportedCultureSelected(FCult
 		SelectedNewCulture = SelectedCulture;
 		NewEntryIndexToBeInitialized = NewElementIndex;
 
+		if (NativeCultureIndexPropertyHandle.IsValid() && NativeCultureIndexPropertyHandle->IsValidHandle())
+		{
+			int32 NativeCultureIndex;
+			NativeCultureIndexPropertyHandle->GetValue(NativeCultureIndex);
+			if (NativeCultureIndex == INDEX_NONE)
+			{
+				NativeCultureIndex = NewElementIndex;
+				NativeCultureIndexPropertyHandle->SetValue(NativeCultureIndex);
+			}
+		}
+
 		// Refresh UI.
 		if (SupportedCulturePicker.IsValid())
 		{
diff --git a/Engine/Source/Editor/LocalizationDashboard/Private/LocalizationTargetSetDetailCustomization.cpp b/Engine/Source/Editor/LocalizationDashboard/Private/LocalizationTargetSetDetailCustomization.cpp
--- a/Engine/Source/Editor/LocalizationDashboard/Private/LocalizationTargetSetDetailCustomization.cpp
+++ b/Engine/Source/Editor/LocalizationDashboard/Private/LocalizationTargetSetDetailCustomization.cpp
@@ -197,7 +197,6 @@ void FLocalizationTargetSetDetailCustomization::RebuildTargetsList()
 			TargetObjectsPropertyHandle->AccessRawData(RawData);
 			void* RawDatum = RawData.Top();
 			TArray<ULocalizationTarget*>& TargetObjectsArray = *(reinterpret_cast< TArray<ULocalizationTarget*>* >(RawDatum));
-
 			FString NewTargetName = "NewTarget";
 			auto TargetNameIsIdentical = [&](ULocalizationTarget* Target) -> bool
 			{
@@ -210,7 +209,8 @@ void FLocalizationTargetSetDetailCustomization::RebuildTargetsList()
 			}
 
 			NewTarget->Settings.Name = NewTargetName;
-			NewTarget->Settings.NativeCultureStatistics.CultureName = FInternationalization::Get().GetCurrentCulture()->GetName();
+			const int32 NativeCultureIndex = NewTarget->Settings.SupportedCulturesStatistics.Add( FCultureStatistics(FInternationalization::Get().GetCurrentCulture()->GetName()) );
+			NewTarget->Settings.NativeCultureIndex = NativeCultureIndex;
 
 			const UObject* SetValue = NewTarget;
 			TargetObjectPropertyHandle->SetValue(SetValue);
diff --git a/Engine/Source/Editor/LocalizationDashboard/Private/LocalizationTargetTypes.cpp b/Engine/Source/Editor/LocalizationDashboard/Private/LocalizationTargetTypes.cpp
--- a/Engine/Source/Editor/LocalizationDashboard/Private/LocalizationTargetTypes.cpp
+++ b/Engine/Source/Editor/LocalizationDashboard/Private/LocalizationTargetTypes.cpp
@@ -219,12 +219,6 @@ bool FGatherTextFromMetaDataConfiguration::Validate(const FString& RootDirectory
 	return false;
 }
 
-FLocalizationTargetSettings::FLocalizationTargetSettings()
-	: Guid(FGuid::NewGuid())
-	, ConflictStatus(ELocalizationTargetConflictStatus::Unknown)
-{
-}
-
 namespace
 {
 	class FWordCountCSVParser
@@ -236,7 +230,10 @@ namespace
 		{
 			// Load CSV file to string.
 			FString CSVString;
-			FFileHelper::LoadFileToString(CSVString, *CSVFilePath);
+			if (!FFileHelper::LoadFileToString(CSVString, *CSVFilePath))
+			{
+				return false;
+			}
 
 			// Parse CSV file string.
 			FCsvParser CSVParser(CSVString);
@@ -258,12 +255,6 @@ namespace
 				return false;
 			}
 
-			uint32* NativeWordCount = WordCounts.Find(Target.NativeCultureStatistics.CultureName);
-			if (NativeWordCount)
-			{
-				Target.NativeCultureStatistics.WordCount = *NativeWordCount;
-			}
-
 			for (FCultureStatistics& SupportedCultureStatistics : Target.SupportedCulturesStatistics)
 			{
 				uint32* SupportedWordCount = WordCounts.Find(SupportedCultureStatistics.CultureName);
@@ -344,12 +335,6 @@ namespace
 				return false;
 			}
 
-			// Parse native culture word count.
-			{
-				const uint32 WordCount = static_cast<uint32>( FCString::Atoi(WordCountRow[1]) );
-				OutCultureWordCounts.Add(Target.NativeCultureStatistics.CultureName, WordCount);
-			}
-
 			// Parse supported culture word counts.
 			for (int32 j = MandatoryColumnNames.Num(); j < WordCountRow.Num(); ++j)
 			{
@@ -389,7 +374,15 @@ bool ULocalizationTarget::IsMemberOfEngineTargetSet() const
 bool ULocalizationTarget::UpdateWordCountsFromCSV()
 {
 	const FString CSVFilePath = LocalizationConfigurationScript::GetWordCountCSVPath(this);
-	return FWordCountCSVParser::Execute(Settings, CSVFilePath);
+	const bool Succeeded = FWordCountCSVParser::Execute(Settings, CSVFilePath);
+	if (!Succeeded)
+	{
+		for (FCultureStatistics& CultureStatistics : Settings.SupportedCulturesStatistics)
+		{
+			CultureStatistics.WordCount = 0;
+		}
+	}
+	return Succeeded;
 }
 
 void ULocalizationTarget::UpdateStatusFromConflictReport()
@@ -476,18 +469,11 @@ bool ULocalizationTarget::RenameTargetAndFiles(const FString& NewName)
 		NamedPaths.Add(LocalizationConfigurationScript::GetWordCountCSVPath(this));
 		NamedPaths.Add(LocalizationConfigurationScript::GetConflictReportPath(this));
 
-		TArray<const FCultureStatistics*> Cultures;
-		Cultures.Add(&Settings.NativeCultureStatistics);
 		for (const FCultureStatistics& Culture : Settings.SupportedCulturesStatistics)
 		{
-			Cultures.Add(&Culture);
-		}
-
-		for (const FCultureStatistics* Culture : Cultures)
-		{
-			NamedPaths.Add(LocalizationConfigurationScript::GetArchivePath(this, Culture->CultureName));
-			NamedPaths.Add(LocalizationConfigurationScript::GetDefaultPOPath(this, Culture->CultureName));
-			NamedPaths.Add(LocalizationConfigurationScript::GetLocResPath(this, Culture->CultureName));
+			NamedPaths.Add(LocalizationConfigurationScript::GetArchivePath(this, Culture.CultureName));
+			NamedPaths.Add(LocalizationConfigurationScript::GetDefaultPOPath(this, Culture.CultureName));
+			NamedPaths.Add(LocalizationConfigurationScript::GetLocResPath(this, Culture.CultureName));
 		}
 
 		NamedPaths.Add(LocalizationConfigurationScript::GetDataDirectory(this));
diff --git a/Engine/Source/Editor/LocalizationDashboard/Private/LocalizationTargetTypes.h b/Engine/Source/Editor/LocalizationDashboard/Private/LocalizationTargetTypes.h
--- a/Engine/Source/Editor/LocalizationDashboard/Private/LocalizationTargetTypes.h
+++ b/Engine/Source/Editor/LocalizationDashboard/Private/LocalizationTargetTypes.h
@@ -208,6 +208,13 @@ struct FCultureStatistics
 	{
 	}
 
+	FCultureStatistics(const FString& InCultureName)
+		: CultureName(InCultureName)
+		, WordCount(0)
+	{
+	}
+
+
 	/* The ISO name for this culture. */
 	UPROPERTY(config, EditAnywhere, Category = "Culture")
 	FString CultureName;
@@ -233,7 +240,12 @@ struct FLocalizationTargetSettings
 {
 	GENERATED_USTRUCT_BODY()
 
-	FLocalizationTargetSettings();
+	FLocalizationTargetSettings::FLocalizationTargetSettings()
+		: Guid(FGuid::NewGuid())
+		, ConflictStatus(ELocalizationTargetConflictStatus::Unknown)
+		, NativeCultureIndex(INDEX_NONE)
+	{
+	}
 
 	/* Unique name for the target. */
 	UPROPERTY(config, EditAnywhere, Category = "Target")
@@ -270,9 +282,9 @@ struct FLocalizationTargetSettings
 	UPROPERTY(config, EditAnywhere, Category = "Gather")
 	FGatherTextFromMetaDataConfiguration GatherFromMetaData;
 
-	/* The culture in which the source text is written in. */
+	/* The index of the native culture among the supported cultures. */
 	UPROPERTY(config, EditAnywhere, Category = "Cultures")
-	FCultureStatistics NativeCultureStatistics;
+	int32 NativeCultureIndex;
 
 	/* Cultures for which the source text is being localized for.*/
 	UPROPERTY(config, EditAnywhere, Category = "Cultures")
diff --git a/Engine/Source/Editor/LocalizationDashboard/Private/SLocalizationDashboardTargetRow.cpp b/Engine/Source/Editor/LocalizationDashboard/Private/SLocalizationDashboardTargetRow.cpp
--- a/Engine/Source/Editor/LocalizationDashboard/Private/SLocalizationDashboardTargetRow.cpp
+++ b/Engine/Source/Editor/LocalizationDashboard/Private/SLocalizationDashboardTargetRow.cpp
@@ -207,10 +207,9 @@ FText SLocalizationDashboardTargetRow::GetCulturesText() const
 	if (TargetSettings)
 	{
 		TArray<FString> OrderedCultureNames;
-		OrderedCultureNames.Add(TargetSettings->NativeCultureStatistics.CultureName);
-		for (auto& ForeignCultureStatistics : TargetSettings->SupportedCulturesStatistics)
+		for (const auto& CultureStatistics : TargetSettings->SupportedCulturesStatistics)
 		{
-			OrderedCultureNames.Add(ForeignCultureStatistics.CultureName);
+			OrderedCultureNames.Add(CultureStatistics.CultureName);
 		}
 
 		FString Result;
@@ -239,13 +238,13 @@ FText SLocalizationDashboardTargetRow::GetCulturesText() const
 uint32 SLocalizationDashboardTargetRow::GetWordCount() const
 {
 	FLocalizationTargetSettings* const TargetSettings = GetTargetSettings();
-	return TargetSettings ? TargetSettings->NativeCultureStatistics.WordCount : 0;
+	return TargetSettings && TargetSettings->SupportedCulturesStatistics.IsValidIndex(TargetSettings->NativeCultureIndex) ? TargetSettings->SupportedCulturesStatistics[TargetSettings->NativeCultureIndex].WordCount : 0;
 }
 
 uint32 SLocalizationDashboardTargetRow::GetNativeWordCount() const
 {
 	FLocalizationTargetSettings* const TargetSettings = GetTargetSettings();
-	return TargetSettings ? TargetSettings->NativeCultureStatistics.WordCount : 0;
+	return TargetSettings && TargetSettings->SupportedCulturesStatistics.IsValidIndex(TargetSettings->NativeCultureIndex) ? TargetSettings->SupportedCulturesStatistics[TargetSettings->NativeCultureIndex].WordCount : 0;
 }
 
 FText SLocalizationDashboardTargetRow::GetWordCountText() const
@@ -458,15 +457,12 @@ void SLocalizationDashboardTargetRow::Delete()
 		ULocalizationTarget* const LocalizationTarget = GetTarget();
 		if (LocalizationTarget)
 		{
-			FText TitleText;
-			FText MessageText;
-
 			// Setup dialog for deleting target.
 			const FText FormatPattern = NSLOCTEXT("LocalizationDashboard", "DeleteTargetConfirmationDialogMessage", "This action can not be undone and data will be permanently lost. Are you sure you would like to delete {TargetName}?");
 			FFormatNamedArguments Arguments;
 			Arguments.Add(TEXT("TargetName"), FText::FromString(LocalizationTarget->Settings.Name));
-			MessageText = FText::Format(FormatPattern, Arguments);
-			TitleText = NSLOCTEXT("LocalizationDashboard", "DeleteTargetConfirmationDialogTitle", "Confirm Target Deletion");
+			const FText MessageText = FText::Format(FormatPattern, Arguments);
+			const FText TitleText = NSLOCTEXT("LocalizationDashboard", "DeleteTargetConfirmationDialogTitle", "Confirm Target Deletion");
 
 			switch(FMessageDialog::Open(EAppMsgType::OkCancel, MessageText, &TitleText))
 			{
diff --git a/Engine/Source/Editor/LocalizationDashboard/Private/SLocalizationTargetEditorCultureRow.cpp b/Engine/Source/Editor/LocalizationDashboard/Private/SLocalizationTargetEditorCultureRow.cpp
--- a/Engine/Source/Editor/LocalizationDashboard/Private/SLocalizationTargetEditorCultureRow.cpp
+++ b/Engine/Source/Editor/LocalizationDashboard/Private/SLocalizationTargetEditorCultureRow.cpp
@@ -21,33 +21,19 @@ TSharedRef<SWidget> SLocalizationTargetEditorCultureRow::GenerateWidgetForColumn
 {
 	TSharedPtr<SWidget> Return;
 
-	if (ColumnName == "Culture")
+	if (ColumnName == "IsNative")
+	{
+		Return = SNew(SCheckBox)
+			.Style(&FCoreStyle::Get().GetWidgetStyle< FCheckBoxStyle >("RadioButton"))
+			.IsChecked_Lambda([this](){return IsNativeCultureForTarget() ? ECheckBoxState::Checked : ECheckBoxState::Unchecked;})
+			.OnCheckStateChanged(this, &SLocalizationTargetEditorCultureRow::OnNativeCultureCheckStateChanged);
+	}
+	else if (ColumnName == "Culture")
 	{
-		const auto& NativeCultureIconVisibilityLambda = [this]()
-		{
-			return IsNativeCultureForTarget() ? EVisibility::Visible : EVisibility::Hidden;
-		};
 		// Culture Name
-		return SNew(SHorizontalBox)
-			+SHorizontalBox::Slot()
-			.AutoWidth()
-			.HAlign(HAlign_Center)
-			.VAlign(VAlign_Center)
-			.Padding(4.0f)
-			[
-				SNew(SImage)
-				.Visibility_Lambda(NativeCultureIconVisibilityLambda)
-				.Image(FEditorStyle::GetBrush("LocalizationTargetEditor.NativeCulture"))
-			]
-		+SHorizontalBox::Slot()
-			.AutoWidth()
-			.HAlign(HAlign_Left)
-			.VAlign(VAlign_Center)
-			[
-				SNew(STextBlock)
-				.Text(this, &SLocalizationTargetEditorCultureRow::GetCultureDisplayName)
-				.ToolTipText(this, &SLocalizationTargetEditorCultureRow::GetCultureName)
-			];
+		Return = SNew(STextBlock)
+			.Text(this, &SLocalizationTargetEditorCultureRow::GetCultureDisplayName)
+			.ToolTipText(this, &SLocalizationTargetEditorCultureRow::GetCultureName);
 	}
 	else if(ColumnName == "WordCount")
 	{
@@ -174,14 +160,7 @@ FCultureStatistics* SLocalizationTargetEditorCultureRow::GetCultureStatistics()
 	FLocalizationTargetSettings* const TargetSettings = GetTargetSettings();
 	if (TargetSettings)
 	{
-		if (CultureIndex == INDEX_NONE)
-		{
-			return &TargetSettings->NativeCultureStatistics;
-		}
-		else
-		{
-			return &TargetSettings->SupportedCulturesStatistics[CultureIndex];
-		}
+		return &TargetSettings->SupportedCulturesStatistics[CultureIndex];
 	}
 	return nullptr;
 }
@@ -199,7 +178,42 @@ FCulturePtr SLocalizationTargetEditorCultureRow::GetCulture() const
 bool SLocalizationTargetEditorCultureRow::IsNativeCultureForTarget() const
 {
 	FLocalizationTargetSettings* const TargetSettings = GetTargetSettings();
-	return &(TargetSettings->NativeCultureStatistics) == GetCultureStatistics();
+	return TargetSettings->SupportedCulturesStatistics.IsValidIndex(TargetSettings->NativeCultureIndex) && &(TargetSettings->SupportedCulturesStatistics[TargetSettings->NativeCultureIndex]) == GetCultureStatistics();
+}
+
+void SLocalizationTargetEditorCultureRow::OnNativeCultureCheckStateChanged( const ECheckBoxState CheckState )
+{
+	FLocalizationTargetSettings* const TargetSettings = GetTargetSettings();
+	
+	if (TargetSettings && TargetSettingsPropertyHandle.IsValid() && TargetSettingsPropertyHandle->IsValidHandle() && CheckState == ECheckBoxState::Checked)
+	{
+		const FText FormatPattern = LOCTEXT("ChangingNativeCultureWarningMessage", "This action can not be undone and all translations be permanently lost. Are you sure you would like to set the native culture to {CultureName}?");
+		FFormatNamedArguments Arguments;
+		Arguments.Add(TEXT("CultureName"), GetCultureDisplayName());
+		const FText MessageText = FText::Format(FormatPattern, Arguments);
+		const FText TitleText = LOCTEXT("ChangingNativeCultureWarningDialogTitle", "Change Native Culture");
+
+		switch(FMessageDialog::Open(EAppMsgType::YesNo, MessageText, &TitleText))
+		{
+		case EAppReturnType::Yes:
+			{
+				ULocalizationTarget* const LocalizationTarget = GetTarget();
+				if (LocalizationTarget)
+				{
+					// Delete data files.
+					const FString DataDirectory = LocalizationConfigurationScript::GetDataDirectory(LocalizationTarget);
+					IFileManager::Get().DeleteDirectory(*DataDirectory, false, true);
+				}
+
+				UpdateTargetFromReports();
+
+				TargetSettingsPropertyHandle->NotifyPreChange();
+				TargetSettings->NativeCultureIndex = CultureIndex;
+				TargetSettingsPropertyHandle->NotifyPostChange();
+			}
+			break;
+		}
+	}
 }
 
 uint32 SLocalizationTargetEditorCultureRow::GetWordCount() const
@@ -215,7 +229,7 @@ uint32 SLocalizationTargetEditorCultureRow::GetWordCount() const
 uint32 SLocalizationTargetEditorCultureRow::GetNativeWordCount() const
 {
 	FLocalizationTargetSettings* const TargetSettings = GetTargetSettings();
-	return TargetSettings ? TargetSettings->NativeCultureStatistics.WordCount : 0;
+	return TargetSettings && TargetSettings->SupportedCulturesStatistics.IsValidIndex(TargetSettings->NativeCultureIndex) ? TargetSettings->SupportedCulturesStatistics[TargetSettings->NativeCultureIndex].WordCount : 0;
 }
 
 FText SLocalizationTargetEditorCultureRow::GetCultureDisplayName() const
@@ -255,16 +269,6 @@ void SLocalizationTargetEditorCultureRow::UpdateTargetFromReports()
 
 		if (TargetSettingsPropertyHandle.IsValid() && TargetSettingsPropertyHandle->IsValidHandle())
 		{
-			const TSharedPtr<IPropertyHandle> NativeCultureStatisticsPropertyHandle = TargetSettingsPropertyHandle->GetChildHandle(GET_MEMBER_NAME_CHECKED(FLocalizationTargetSettings, NativeCultureStatistics));
-			if (NativeCultureStatisticsPropertyHandle.IsValid() && NativeCultureStatisticsPropertyHandle->IsValidHandle())
-			{
-				const TSharedPtr<IPropertyHandle> WordCountPropertyHandle = NativeCultureStatisticsPropertyHandle->GetChildHandle(GET_MEMBER_NAME_CHECKED(FCultureStatistics, WordCount));
-				if (WordCountPropertyHandle.IsValid() && WordCountPropertyHandle->IsValidHandle())
-				{
-					WordCountPropertyHandles.Add(WordCountPropertyHandle);
-				}
-			}
-
 			const TSharedPtr<IPropertyHandle> SupportedCulturesStatisticsPropertyHandle = TargetSettingsPropertyHandle->GetChildHandle(GET_MEMBER_NAME_CHECKED(FLocalizationTargetSettings, SupportedCulturesStatistics));
 			if (SupportedCulturesStatisticsPropertyHandle.IsValid() && SupportedCulturesStatisticsPropertyHandle->IsValidHandle())
 			{
diff --git a/Engine/Source/Editor/LocalizationDashboard/Private/SLocalizationTargetEditorCultureRow.h b/Engine/Source/Editor/LocalizationDashboard/Private/SLocalizationTargetEditorCultureRow.h
--- a/Engine/Source/Editor/LocalizationDashboard/Private/SLocalizationTargetEditorCultureRow.h
+++ b/Engine/Source/Editor/LocalizationDashboard/Private/SLocalizationTargetEditorCultureRow.h
@@ -19,6 +19,7 @@ private:
 	FCulturePtr GetCulture() const;
 	bool IsNativeCultureForTarget() const;
 
+	void OnNativeCultureCheckStateChanged(const ECheckBoxState CheckState);
 	uint32 GetWordCount() const;
 	uint32 GetNativeWordCount() const;
 
