Removed "Max" value from showing up in the localization target editor loading policy dropdown.

diff --git a/Engine/Source/Editor/LocalizationDashboard/Private/LocalizationTargetDetailCustomization.cpp b/Engine/Source/Editor/LocalizationDashboard/Private/LocalizationTargetDetailCustomization.cpp
--- a/Engine/Source/Editor/LocalizationDashboard/Private/LocalizationTargetDetailCustomization.cpp
+++ b/Engine/Source/Editor/LocalizationDashboard/Private/LocalizationTargetDetailCustomization.cpp
@@ -310,7 +310,7 @@ void FLocalizationTargetDetailCustomization::CustomizeDetails(IDetailLayoutBuild
 		{
 			UEnum* const LoadingPolicyEnum = FindObjectChecked<UEnum>(ANY_PACKAGE, TEXT("ELocalizationTargetLoadingPolicy"));
 			TArray< TSharedPtr<ELocalizationTargetLoadingPolicy> > Array;
-			for (int32 i = 0; i < LoadingPolicyEnum->NumEnums(); ++i)
+			for (int32 i = 0; i < LoadingPolicyEnum->NumEnums() - 1; ++i)
 			{
 				Array.Add( MakeShareable( new ELocalizationTargetLoadingPolicy(static_cast<ELocalizationTargetLoadingPolicy>(i)) ) );
 			}
