Removed the method declarations of some deleted methods from the localization target detail customizations.

diff --git a/Engine/Source/Editor/LocalizationDashboard/Private/LocalizationTargetDetailCustomization.h b/Engine/Source/Editor/LocalizationDashboard/Private/LocalizationTargetDetailCustomization.h
--- a/Engine/Source/Editor/LocalizationDashboard/Private/LocalizationTargetDetailCustomization.h
+++ b/Engine/Source/Editor/LocalizationDashboard/Private/LocalizationTargetDetailCustomization.h
@@ -33,10 +33,6 @@ private:
 	void OnTargetDependencyCheckStateChanged(ULocalizationTarget* const OtherLocalizationTarget, const ECheckBoxState State);
 	ECheckBoxState IsTargetDependencyChecked(ULocalizationTarget* const OtherLocalizationTarget) const;
 
-	FText GetNativeCultureName() const;
-	FText GetNativeCultureDisplayName() const;
-	void OnNativeCultureSelected(FCulturePtr SelectedCulture, ESelectInfo::Type SelectInfo);
-
 	void Gather();
 	void ImportAllCultures();
 	void ExportAllCultures();
@@ -66,11 +62,9 @@ private:
 	TArray<ULocalizationTarget*> TargetDependenciesOptionsList;
 	TSharedPtr< SListView<ULocalizationTarget*> > TargetDependenciesListView;
 
-	TSharedPtr<IPropertyHandle> NativeCultureStatisticsPropertyHandle;
-	TSharedPtr<IPropertyHandle> NativeCultureNamePropertyHandle;
-	TSharedPtr<SComboButton> NativeCultureComboButton;
 	TArray<FCulturePtr> AllCultures;
 
+	TSharedPtr<IPropertyHandle> NativeCultureIndexPropertyHandle;
 	TSharedPtr<IPropertyHandle> SupportedCulturesStatisticsPropertyHandle;
 	FSimpleDelegate SupportedCulturesStatisticsPropertyHandle_OnNumElementsChanged;
 	TSharedPtr< SListView< TSharedPtr<IPropertyHandle> > > SupportedCultureListView;
