Removed redundant scoping of constructor in FLocalizationTargetSettings that was causing warnings.

Fixes this error:
 mono: Editor/LocalizationDashboard/Private/LocalizationTargetTypes.h:243:31: error: extra qualification on member 'FLocalizationTargetSettings'

diff --git a/Engine/Source/Editor/LocalizationDashboard/Private/LocalizationTargetTypes.h b/Engine/Source/Editor/LocalizationDashboard/Private/LocalizationTargetTypes.h
--- a/Engine/Source/Editor/LocalizationDashboard/Private/LocalizationTargetTypes.h
+++ b/Engine/Source/Editor/LocalizationDashboard/Private/LocalizationTargetTypes.h
@@ -240,7 +240,7 @@ struct FLocalizationTargetSettings
 {
 	GENERATED_USTRUCT_BODY()
 
-	FLocalizationTargetSettings::FLocalizationTargetSettings()
+	FLocalizationTargetSettings()
 		: Guid(FGuid::NewGuid())
 		, ConflictStatus(ELocalizationTargetConflictStatus::Unknown)
 		, NativeCultureIndex(INDEX_NONE)
