Fixed a reported bug with the localization dashboard not properly executing commandlets if they project file path includes spaces.

diff --git a/Engine/Source/Editor/LocalizationDashboard/Private/LocalizationCommandletExecution.cpp b/Engine/Source/Editor/LocalizationDashboard/Private/LocalizationCommandletExecution.cpp
--- a/Engine/Source/Editor/LocalizationDashboard/Private/LocalizationCommandletExecution.cpp
+++ b/Engine/Source/Editor/LocalizationDashboard/Private/LocalizationCommandletExecution.cpp
@@ -722,7 +722,8 @@ TSharedPtr<FLocalizationCommandletProcess> FLocalizationCommandletProcess::Execu
 		}
 	}
 
-	const FString ProcessArguments = CommandletHelpers::BuildCommandletProcessArguments(TEXT("GatherText"), UseProjectFile ? *FPaths::GetProjectFilePath() : nullptr, *CommandletArguments);
+	const FString ProjectFilePath = FString::Printf(TEXT("\"%s\""), *FPaths::ConvertRelativePathToFull(FPaths::GetProjectFilePath()));
+	const FString ProcessArguments = CommandletHelpers::BuildCommandletProcessArguments(TEXT("GatherText"), UseProjectFile ? *ProjectFilePath : nullptr, *CommandletArguments);
 	FProcHandle CommandletProcessHandle = FPlatformProcess::CreateProc(*FUnrealEdMisc::Get().GetExecutableForCommandlets(), *ProcessArguments, true, true, true, nullptr, 0, nullptr, WritePipe);
 
 	// Close pipes if process failed.
