Updated base editor INI so that the default game localization target has English as the native culture. Added NativeCulture setting to localization archive generation commandlet settings and deprecated SourceCulture. Added logic to enable and disable localization dashboard operations if a target doesn't have native/supported foreign cultures specified. Added localization dashboard operations for refreshing word counts and compiling culture localization data where they were not consistently available.

diff --git a/Engine/Source/Editor/EditorStyle/Private/SlateEditorStyle.cpp b/Engine/Source/Editor/EditorStyle/Private/SlateEditorStyle.cpp
--- a/Engine/Source/Editor/EditorStyle/Private/SlateEditorStyle.cpp
+++ b/Engine/Source/Editor/EditorStyle/Private/SlateEditorStyle.cpp
@@ -1,4 +1,4 @@
-// Copyright 1998-2015 Epic Games, Inc. All Rights Reserved.
+﻿// Copyright 1998-2015 Epic Games, Inc. All Rights Reserved.
 
 #include "EditorStylePrivatePCH.h"
 
@@ -6624,25 +6624,27 @@ void FSlateEditorStyle::FStyle::SetupTranslationEditorStyles()
 	Set("TranslationEditor.TranslationPicker", new IMAGE_BRUSH("Icons/icon_StaticMeshEd_VertColor_40x", Icon40x40));
 }
 
+
 void FSlateEditorStyle::FStyle::SetupLocalizationDashboardStyles()
 {
 	Set("LocalizationDashboard.GatherAllTargets", new IMAGE_BRUSH("Icons/Icon_Localisation_Gather_All_40x", Icon40x40));
-	Set("LocalizationDashboard.ImportForAllTargetsCultures", new IMAGE_BRUSH("Icons/Icon_Localisation_Import_All_40x", Icon40x40));
-	Set("LocalizationDashboard.ExportForAllTargetsCultures", new IMAGE_BRUSH("Icons/Icon_Localisation_Export_All_40x", Icon40x40));
-	Set("LocalizationDashboard.CompileAllTargets", new IMAGE_BRUSH("Icons/Icon_Localisation_Compile_Translations_40x", Icon40x40));
+	Set("LocalizationDashboard.ImportAllTargetsAllCultures", new IMAGE_BRUSH("Icons/Icon_Localisation_Import_All_40x", Icon40x40));
+	Set("LocalizationDashboard.ExportAllTargetsAllCultures", new IMAGE_BRUSH("Icons/Icon_Localisation_Export_All_40x", Icon40x40));
+	Set("LocalizationDashboard.CountWordsForAllTargets", new IMAGE_BRUSH("Icons/Icon_Localisation_Refresh_Word_Counts_40x", Icon40x40));
+	Set("LocalizationDashboard.CompileAllTargetsAllCultures", new IMAGE_BRUSH("Icons/Icon_Localisation_Compile_Translations_40x", Icon40x40));
 
 	Set("LocalizationDashboard.GatherTarget", new IMAGE_BRUSH("Icons/Icon_Localisation_Gather_All_16x", Icon16x16));
-	Set("LocalizationDashboard.RefreshWordCount", new IMAGE_BRUSH("Icons/Icon_Localisation_Refresh_Word_Counts_16x", Icon16x16));
-	Set("LocalizationDashboard.ImportForAllCultures", new IMAGE_BRUSH("Icons/Icon_Localisation_Import_All_16x", Icon16x16));
-	Set("LocalizationDashboard.ExportForAllCultures", new IMAGE_BRUSH("Icons/Icon_Localisation_Export_All_16x", Icon16x16));
-	Set("LocalizationDashboard.CompileTarget", new IMAGE_BRUSH("Icons/Icon_Localisation_Compile_Translations_16x", Icon16x16));
+	Set("LocalizationDashboard.ImportAllCultures", new IMAGE_BRUSH("Icons/Icon_Localisation_Import_All_16x", Icon16x16));
+	Set("LocalizationDashboard.ExportAllCultures", new IMAGE_BRUSH("Icons/Icon_Localisation_Export_All_16x", Icon16x16));
+	Set("LocalizationDashboard.CountWordsForTarget", new IMAGE_BRUSH("Icons/Icon_Localisation_Refresh_Word_Counts_16x", Icon16x16));
+	Set("LocalizationDashboard.CompileAllCultures", new IMAGE_BRUSH("Icons/Icon_Localisation_Compile_Translations_16x", Icon16x16));
 	Set("LocalizationDashboard.DeleteTarget", new IMAGE_BRUSH("Icons/Cross_12x", Icon12x12 ) );
 
 	Set("LocalizationTargetEditor.Gather", new IMAGE_BRUSH("Icons/Icon_Localisation_Gather_All_40x", Icon40x40));
-	Set("LocalizationTargetEditor.ImportForAllCultures", new IMAGE_BRUSH("Icons/Icon_Localisation_Import_All_40x", Icon40x40));
-	Set("LocalizationTargetEditor.ExportForAllCultures", new IMAGE_BRUSH("Icons/Icon_Localisation_Export_All_40x", Icon40x40));
-	Set("LocalizationTargetEditor.RefreshWordCounts", new IMAGE_BRUSH("Icons/Icon_Localisation_Refresh_Word_Counts_40x", Icon40x40));
-	Set("LocalizationTargetEditor.Compile", new IMAGE_BRUSH( "Icons/Icon_Localisation_Compile_Translations_40x", Icon40x40));
+	Set("LocalizationTargetEditor.ImportAllCultures", new IMAGE_BRUSH("Icons/Icon_Localisation_Import_All_40x", Icon40x40));
+	Set("LocalizationTargetEditor.ExportAllCultures", new IMAGE_BRUSH("Icons/Icon_Localisation_Export_All_40x", Icon40x40));
+	Set("LocalizationTargetEditor.CountWords", new IMAGE_BRUSH("Icons/Icon_Localisation_Refresh_Word_Counts_40x", Icon40x40));
+	Set("LocalizationTargetEditor.CompileAllCultures", new IMAGE_BRUSH( "Icons/Icon_Localisation_Compile_Translations_40x", Icon40x40));
 
 	Set("LocalizationTargetEditor.DirectoryPicker", new IMAGE_BRUSH( "Icons/ellipsis_12x", Icon12x12 ));
 	Set("LocalizationTargetEditor.GatherSettingsIcon_Valid", new IMAGE_BRUSH("Settings/Settings_Good", Icon16x16));
@@ -6651,8 +6653,9 @@ void FSlateEditorStyle::FStyle::SetupLocalizationDashboardStyles()
 	Set("LocalizationTargetEditor.NativeCulture", new IMAGE_BRUSH( "Icons/Star_16x", Icon16x16 ) );
 
 	Set("LocalizationTargetEditor.EditTranslations", new IMAGE_BRUSH("Icons/icon_file_open_16px", Icon16x16));
-	Set("LocalizationTargetEditor.ImportForCulture", new IMAGE_BRUSH("Icons/Icon_Localisation_Import_All_16x", Icon16x16));
-	Set("LocalizationTargetEditor.ExportForCulture", new IMAGE_BRUSH("Icons/Icon_Localisation_Export_All_16x", Icon16x16));
+	Set("LocalizationTargetEditor.ImportCulture", new IMAGE_BRUSH("Icons/Icon_Localisation_Import_All_16x", Icon16x16));
+	Set("LocalizationTargetEditor.ExportCulture", new IMAGE_BRUSH("Icons/Icon_Localisation_Export_All_16x", Icon16x16));
+	Set("LocalizationTargetEditor.CompileCulture", new IMAGE_BRUSH("Icons/Icon_Localisation_Compile_Translations_16x", Icon16x16));
 	Set("LocalizationTargetEditor.DeleteCulture", new IMAGE_BRUSH("Icons/Cross_12x", Icon12x12 ) );
 
 	Set("LocalizationTargetEditor.GatherSettings.AddMetaDataTextKeyPatternArgument", new IMAGE_BRUSH("Icons/icon_Blueprint_AddVariable_40px", Icon16x16 ) );
diff --git a/Engine/Source/Editor/LocalizationDashboard/Private/LocalizationCommandletTasks.cpp b/Engine/Source/Editor/LocalizationDashboard/Private/LocalizationCommandletTasks.cpp
--- a/Engine/Source/Editor/LocalizationDashboard/Private/LocalizationCommandletTasks.cpp
+++ b/Engine/Source/Editor/LocalizationDashboard/Private/LocalizationCommandletTasks.cpp
@@ -258,4 +258,31 @@ bool LocalizationCommandletTasks::CompileTarget(const TSharedRef<SWindow>& Paren
 	return LocalizationCommandletExecution::Execute(ParentWindow, WindowTitle, Tasks);
 }
 
+bool LocalizationCommandletTasks::CompileCulture(const TSharedRef<SWindow>& ParentWindow, ULocalizationTarget* const Target, const FString& CultureName)
+{
+	FCulturePtr Culture = FInternationalization::Get().GetCulture(CultureName);
+	if (!Culture.IsValid())
+	{
+		return false;
+	}
+
+	TArray<LocalizationCommandletExecution::FTask> Tasks;
+	const bool ShouldUseProjectFile = !Target->IsMemberOfEngineTargetSet();
+
+	const FString DefaultCompileScriptPath = LocalizationConfigurationScript::GetCompileScriptPath(Target, TOptional<FString>(CultureName));
+	const FString CompileScriptPath = FPaths::CreateTempFilename(*FPaths::GetPath(DefaultCompileScriptPath), *FPaths::GetBaseFilename(DefaultCompileScriptPath), *FPaths::GetExtension(DefaultCompileScriptPath, true));
+	LocalizationConfigurationScript::GenerateCompileScript(Target, TOptional<FString>(CultureName)).Write(CompileScriptPath);
+	Tasks.Add(LocalizationCommandletExecution::FTask(LOCTEXT("CompileTaskName", "Compile Translations"), CompileScriptPath, ShouldUseProjectFile));
+
+	FFormatNamedArguments Arguments;
+	Arguments.Add(TEXT("CultureName"), FText::FromString(Culture->GetDisplayName()));
+	Arguments.Add(TEXT("TargetName"), FText::FromString(Target->Settings.Name));
+
+	const FText WindowTitle = FText::Format(LOCTEXT("CompileCultureForTargetWindowTitle", "Compile {CultureName} Translations for Target {TargetName}"), Arguments);
+
+	bool HasSucceeeded = LocalizationCommandletExecution::Execute(ParentWindow, WindowTitle, Tasks);
+	IFileManager::Get().Delete(*CompileScriptPath); // Don't clutter up the loc config directory with scripts for individual cultures.
+	return HasSucceeeded;
+}
+
 #undef LOCTEXT_NAMESPACE
\ No newline at end of file
diff --git a/Engine/Source/Editor/LocalizationDashboard/Private/LocalizationCommandletTasks.h b/Engine/Source/Editor/LocalizationDashboard/Private/LocalizationCommandletTasks.h
--- a/Engine/Source/Editor/LocalizationDashboard/Private/LocalizationCommandletTasks.h
+++ b/Engine/Source/Editor/LocalizationDashboard/Private/LocalizationCommandletTasks.h
@@ -21,4 +21,5 @@ namespace LocalizationCommandletTasks
 
 	bool CompileTargets(const TSharedRef<SWindow>& ParentWindow, const TArray<ULocalizationTarget*>& Targets);
 	bool CompileTarget(const TSharedRef<SWindow>& ParentWindow, ULocalizationTarget* const Target);
+	bool CompileCulture(const TSharedRef<SWindow>& ParentWindow, ULocalizationTarget* const Target, const FString& CultureName);
 }
\ No newline at end of file
diff --git a/Engine/Source/Editor/LocalizationDashboard/Private/LocalizationConfigurationScript.cpp b/Engine/Source/Editor/LocalizationDashboard/Private/LocalizationConfigurationScript.cpp
--- a/Engine/Source/Editor/LocalizationDashboard/Private/LocalizationConfigurationScript.cpp
+++ b/Engine/Source/Editor/LocalizationDashboard/Private/LocalizationConfigurationScript.cpp
@@ -130,7 +130,10 @@ namespace LocalizationConfigurationScript
 			ConfigSection.Add( TEXT("ManifestName"), FPaths::GetCleanFilename(GetManifestPath(Target)) );
 			ConfigSection.Add( TEXT("ArchiveName"), FPaths::GetCleanFilename(GetArchivePath(Target, FString())) );
 
-			ConfigSection.Add( TEXT("SourceCulture"), Target->Settings.SupportedCulturesStatistics[Target->Settings.NativeCultureIndex].CultureName );
+			if (Target->Settings.SupportedCulturesStatistics.IsValidIndex(Target->Settings.NativeCultureIndex))
+			{
+			ConfigSection.Add( TEXT("NativeCulture"), Target->Settings.SupportedCulturesStatistics[Target->Settings.NativeCultureIndex].CultureName );
+			}
 			for (const FCultureStatistics& CultureStatistics : Target->Settings.SupportedCulturesStatistics)
 			{
 				ConfigSection.Add( TEXT("CulturesToGenerate"), CultureStatistics.CultureName );
@@ -508,7 +511,7 @@ namespace LocalizationConfigurationScript
 		return GetScriptDirectory(Target) / FString::Printf( TEXT("%s_GenerateReports.%s"), *(Target->Settings.Name), TEXT("ini") );
 	}
 
-	FLocalizationConfigurationScript GenerateCompileScript(const ULocalizationTarget* const Target)
+	FLocalizationConfigurationScript GenerateCompileScript(const ULocalizationTarget* const Target, const TOptional<FString> CultureName)
 	{
 		FLocalizationConfigurationScript Script;
 
@@ -529,9 +532,24 @@ namespace LocalizationConfigurationScript
 			ConfigSection.Add( TEXT("ManifestName"), FString::Printf( TEXT("%s.%s"), *Target->Settings.Name, TEXT("manifest") ) );
 			ConfigSection.Add( TEXT("ResourceName"), FString::Printf( TEXT("%s.%s"), *Target->Settings.Name, TEXT("locres") ) );
 
-			for (const FCultureStatistics& CultureStatistics : Target->Settings.SupportedCulturesStatistics)
+			const auto& AddCultureToGenerate = [&](const int32 Index)
 			{
-				ConfigSection.Add( TEXT("CulturesToGenerate"), CultureStatistics.CultureName );
+				ConfigSection.Add( TEXT("CulturesToGenerate"), Target->Settings.SupportedCulturesStatistics[Index].CultureName );
+			};
+
+			// Compile a specific culture.
+			if (CultureName.IsSet())
+			{
+				const int32 CultureIndex = Target->Settings.SupportedCulturesStatistics.IndexOfByPredicate([CultureName](const FCultureStatistics& Culture) { return Culture.CultureName == CultureName.GetValue(); });
+				AddCultureToGenerate(CultureIndex);
+			}
+			// Compile all cultures.
+			else
+			{
+				for (int32 CultureIndex = 0; CultureIndex < Target->Settings.SupportedCulturesStatistics.Num(); ++CultureIndex)
+			{
+					AddCultureToGenerate(CultureIndex);
+				}
 			}
 		}
 
@@ -540,8 +558,18 @@ namespace LocalizationConfigurationScript
 		return Script;
 	}
 
-	FString GetCompileScriptPath(const ULocalizationTarget* const Target)
+	FString GetCompileScriptPath(const ULocalizationTarget* const Target, const TOptional<FString> CultureName)
+	{
+		const FString ConfigFileDirectory = GetScriptDirectory(Target);
+		FString ConfigFilePath;
+		if (CultureName.IsSet())
+		{
+			ConfigFilePath = ConfigFileDirectory / FString::Printf( TEXT("%s_Compile_%s.%s"), *Target->Settings.Name, *CultureName.GetValue(), TEXT("ini") );
+		}
+		else
 	{
-		return GetScriptDirectory(Target) / FString::Printf( TEXT("%s_Compile.%s"), *(Target->Settings.Name), TEXT("ini") );
+			ConfigFilePath = ConfigFileDirectory / FString::Printf( TEXT("%s_Compile.%s"), *Target->Settings.Name, TEXT("ini") );
+		}
+		return ConfigFilePath;
 	}
 }
\ No newline at end of file
diff --git a/Engine/Source/Editor/LocalizationDashboard/Private/LocalizationConfigurationScript.h b/Engine/Source/Editor/LocalizationDashboard/Private/LocalizationConfigurationScript.h
--- a/Engine/Source/Editor/LocalizationDashboard/Private/LocalizationConfigurationScript.h
+++ b/Engine/Source/Editor/LocalizationDashboard/Private/LocalizationConfigurationScript.h
@@ -46,6 +46,6 @@ namespace LocalizationConfigurationScript
 	FLocalizationConfigurationScript GenerateWordCountReportScript(const ULocalizationTarget* const Target);
 	FString GetWordCountReportScriptPath(const ULocalizationTarget* const Target);
 
-	FLocalizationConfigurationScript GenerateCompileScript(const ULocalizationTarget* const Target);
-	FString GetCompileScriptPath(const ULocalizationTarget* const Target);
+	FLocalizationConfigurationScript GenerateCompileScript(const ULocalizationTarget* const Target, const TOptional<FString> CultureName = TOptional<FString>());
+	FString GetCompileScriptPath(const ULocalizationTarget* const Target, const TOptional<FString> CultureName = TOptional<FString>());
 }
\ No newline at end of file
diff --git a/Engine/Source/Editor/LocalizationDashboard/Private/LocalizationTargetDetailCustomization.cpp b/Engine/Source/Editor/LocalizationDashboard/Private/LocalizationTargetDetailCustomization.cpp
--- a/Engine/Source/Editor/LocalizationDashboard/Private/LocalizationTargetDetailCustomization.cpp
+++ b/Engine/Source/Editor/LocalizationDashboard/Private/LocalizationTargetDetailCustomization.cpp
@@ -14,6 +14,7 @@
 #include "LocalizationCommandletTasks.h"
 #include "ObjectEditorUtils.h"
 #include "LocalizationDashboardSettings.h"
+#include "SErrorText.h"
 
 #define LOCTEXT_NAMESPACE "LocalizationTargetEditor"
 
@@ -56,8 +57,8 @@ public:
 	TSharedPtr<FUICommandInfo> Gather;
 	TSharedPtr<FUICommandInfo> ImportAllCultures;
 	TSharedPtr<FUICommandInfo> ExportAllCultures;
-	TSharedPtr<FUICommandInfo> RefreshWordCounts;
-	TSharedPtr<FUICommandInfo> Compile;
+	TSharedPtr<FUICommandInfo> CountWords;
+	TSharedPtr<FUICommandInfo> CompileAllCultures;
 
 	/** Initialize commands */
 	virtual void RegisterCommands() override;
@@ -68,8 +69,8 @@ void FLocalizationTargetEditorCommands::RegisterCommands()
 	UI_COMMAND( Gather, "Gather", "Gathers text for this target.", EUserInterfaceActionType::Button, FInputChord() );
 	UI_COMMAND( ImportAllCultures, "Import All", "Imports translations for all cultures of this target.", EUserInterfaceActionType::Button, FInputChord() );
 	UI_COMMAND( ExportAllCultures, "Export All", "Exports translations for all cultures of this target.", EUserInterfaceActionType::Button, FInputChord() );
-	UI_COMMAND( RefreshWordCounts, "Count Words", "Recalculates the word counts for all cultures of this target.", EUserInterfaceActionType::Button, FInputChord() );
-	UI_COMMAND( Compile, "Compile", "Compiles translations for this target.", EUserInterfaceActionType::Button, FInputChord() );
+	UI_COMMAND( CountWords, "Count Words", "Recalculates the word counts for this target.", EUserInterfaceActionType::Button, FInputChord() );
+	UI_COMMAND( CompileAllCultures, "Compile All", "Compiles translations for all cultures of this target.", EUserInterfaceActionType::Button, FInputChord() );
 }
 
 void FLocalizationTargetDetailCustomization::CustomizeDetails(IDetailLayoutBuilder& DetailBuilder)
@@ -185,20 +186,24 @@ void FLocalizationTargetDetailCustomization::CustomizeDetails(IDetailLayoutBuild
 		const TSharedRef< FUICommandList > CommandList = MakeShareable(new FUICommandList);
 		FToolBarBuilder ToolBarBuilder( CommandList, FMultiBoxCustomization::AllowCustomization("LocalizationTargetEditor") );
 
-		CommandList->MapAction( FLocalizationTargetEditorCommands::Get().Gather, FExecuteAction::CreateSP(this, &FLocalizationTargetDetailCustomization::Gather));
-		ToolBarBuilder.AddToolBarButton(FLocalizationTargetEditorCommands::Get().Gather, NAME_None, TAttribute<FText>(), TAttribute<FText>(), FSlateIcon(FEditorStyle::GetStyleSetName(), "LocalizationTargetEditor.Gather"));
+		TAttribute<FText> GatherToolTipTextAttribute = TAttribute<FText>::Create(TAttribute<FText>::FGetter::CreateLambda([this]() -> FText
+			{
+				return CanGather() ? FLocalizationTargetEditorCommands::Get().Gather->GetDescription() : LOCTEXT("GatherDisabledToolTip", "Must have a native culture specified in order to gather.");
+			}));
+		CommandList->MapAction( FLocalizationTargetEditorCommands::Get().Gather, FExecuteAction::CreateSP(this, &FLocalizationTargetDetailCustomization::Gather), FCanExecuteAction::CreateSP(this, &FLocalizationTargetDetailCustomization::CanGather));
+		ToolBarBuilder.AddToolBarButton(FLocalizationTargetEditorCommands::Get().Gather, NAME_None, TAttribute<FText>(), GatherToolTipTextAttribute, FSlateIcon(FEditorStyle::GetStyleSetName(), "LocalizationTargetEditor.Gather"));
 
-		CommandList->MapAction( FLocalizationTargetEditorCommands::Get().ImportAllCultures, FExecuteAction::CreateSP(this, &FLocalizationTargetDetailCustomization::ImportAllCultures));
-		ToolBarBuilder.AddToolBarButton(FLocalizationTargetEditorCommands::Get().ImportAllCultures, NAME_None, TAttribute<FText>(), TAttribute<FText>(), FSlateIcon(FEditorStyle::GetStyleSetName(), "LocalizationTargetEditor.ImportForAllCultures"));
+		CommandList->MapAction( FLocalizationTargetEditorCommands::Get().ImportAllCultures, FExecuteAction::CreateSP(this, &FLocalizationTargetDetailCustomization::ImportAllCultures), FCanExecuteAction::CreateSP(this, &FLocalizationTargetDetailCustomization::CanImportAllCultures));
+		ToolBarBuilder.AddToolBarButton(FLocalizationTargetEditorCommands::Get().ImportAllCultures, NAME_None, TAttribute<FText>(), TAttribute<FText>(), FSlateIcon(FEditorStyle::GetStyleSetName(), "LocalizationTargetEditor.ImportAllCultures"));
 
-		CommandList->MapAction( FLocalizationTargetEditorCommands::Get().ExportAllCultures, FExecuteAction::CreateSP(this, &FLocalizationTargetDetailCustomization::ExportAllCultures));
-		ToolBarBuilder.AddToolBarButton(FLocalizationTargetEditorCommands::Get().ExportAllCultures, NAME_None, TAttribute<FText>(), TAttribute<FText>(), FSlateIcon(FEditorStyle::GetStyleSetName(), "LocalizationTargetEditor.ExportForAllCultures"));
+		CommandList->MapAction( FLocalizationTargetEditorCommands::Get().ExportAllCultures, FExecuteAction::CreateSP(this, &FLocalizationTargetDetailCustomization::ExportAllCultures), FCanExecuteAction::CreateSP(this, &FLocalizationTargetDetailCustomization::CanExportAllCultures));
+		ToolBarBuilder.AddToolBarButton(FLocalizationTargetEditorCommands::Get().ExportAllCultures, NAME_None, TAttribute<FText>(), TAttribute<FText>(), FSlateIcon(FEditorStyle::GetStyleSetName(), "LocalizationTargetEditor.ExportAllCultures"));
 
-		CommandList->MapAction( FLocalizationTargetEditorCommands::Get().RefreshWordCounts, FExecuteAction::CreateSP(this, &FLocalizationTargetDetailCustomization::RefreshWordCounts));
-		ToolBarBuilder.AddToolBarButton(FLocalizationTargetEditorCommands::Get().RefreshWordCounts, NAME_None, TAttribute<FText>(), TAttribute<FText>(), FSlateIcon(FEditorStyle::GetStyleSetName(), "LocalizationTargetEditor.RefreshWordCounts"));
+		CommandList->MapAction( FLocalizationTargetEditorCommands::Get().CountWords, FExecuteAction::CreateSP(this, &FLocalizationTargetDetailCustomization::CountWords), FCanExecuteAction::CreateSP(this, &FLocalizationTargetDetailCustomization::CanCountWords));
+		ToolBarBuilder.AddToolBarButton(FLocalizationTargetEditorCommands::Get().CountWords, NAME_None, TAttribute<FText>(), TAttribute<FText>(), FSlateIcon(FEditorStyle::GetStyleSetName(), "LocalizationTargetEditor.CountWords"));
 
-		CommandList->MapAction( FLocalizationTargetEditorCommands::Get().Compile, FExecuteAction::CreateSP(this, &FLocalizationTargetDetailCustomization::Compile));
-		ToolBarBuilder.AddToolBarButton(FLocalizationTargetEditorCommands::Get().Compile, NAME_None, TAttribute<FText>(), TAttribute<FText>(), FSlateIcon(FEditorStyle::GetStyleSetName(), "LocalizationTargetEditor.Compile"));
+		CommandList->MapAction( FLocalizationTargetEditorCommands::Get().CompileAllCultures, FExecuteAction::CreateSP(this, &FLocalizationTargetDetailCustomization::CompileAllCultures), FCanExecuteAction::CreateSP(this, &FLocalizationTargetDetailCustomization::CanCompileAllCultures));
+		ToolBarBuilder.AddToolBarButton(FLocalizationTargetEditorCommands::Get().CompileAllCultures, NAME_None, TAttribute<FText>(), TAttribute<FText>(), FSlateIcon(FEditorStyle::GetStyleSetName(), "LocalizationTargetEditor.CompileAllCultures"));
 
 		BuildListedCulturesList();
 
@@ -247,6 +252,12 @@ void FLocalizationTargetDetailCustomization::CustomizeDetails(IDetailLayoutBuild
 						.FillWidth(0.3f)
 						)
 					]
+				+SVerticalBox::Slot()
+					.AutoHeight()
+					.VAlign(VAlign_Center)
+					[
+						SAssignNew(NoSupportedCulturesErrorText, SErrorText)
+					]
 				+SVerticalBox::Slot()
 					.AutoHeight()
 					.VAlign(VAlign_Center)
@@ -651,6 +662,11 @@ ECheckBoxState FLocalizationTargetDetailCustomization::IsTargetDependencyChecked
 	return LocalizationTarget->Settings.TargetDependencies.Contains(OtherLocalizationTarget->Settings.Guid) ? ECheckBoxState::Checked : ECheckBoxState::Unchecked;
 }
 
+bool FLocalizationTargetDetailCustomization::CanGather() const
+{
+	return LocalizationTarget.IsValid() && LocalizationTarget->Settings.SupportedCulturesStatistics.Num() > 0 && LocalizationTarget->Settings.SupportedCulturesStatistics.IsValidIndex(LocalizationTarget->Settings.NativeCultureIndex);
+}
+
 void FLocalizationTargetDetailCustomization::Gather()
 {
 	if (LocalizationTarget.IsValid())
@@ -688,6 +704,11 @@ void FLocalizationTargetDetailCustomization::Gather()
 	}
 }
 
+bool FLocalizationTargetDetailCustomization::CanImportAllCultures() const
+{
+	return LocalizationTarget.IsValid() && LocalizationTarget->Settings.SupportedCulturesStatistics.Num() > 0;
+}
+
 void FLocalizationTargetDetailCustomization::ImportAllCultures()
 {
 	IDesktopPlatform* DesktopPlatform = FDesktopPlatformModule::Get();
@@ -720,6 +741,11 @@ void FLocalizationTargetDetailCustomization::ImportAllCultures()
 	}
 }
 
+bool FLocalizationTargetDetailCustomization::CanExportAllCultures() const
+{
+	return LocalizationTarget.IsValid() && LocalizationTarget->Settings.SupportedCulturesStatistics.Num() > 0;
+}
+
 void FLocalizationTargetDetailCustomization::ExportAllCultures()
 {
 	IDesktopPlatform* DesktopPlatform = FDesktopPlatformModule::Get();
@@ -750,24 +776,34 @@ void FLocalizationTargetDetailCustomization::ExportAllCultures()
 	}
 }
 
-void FLocalizationTargetDetailCustomization::Compile()
+bool FLocalizationTargetDetailCustomization::CanCountWords() const
+{
+	return LocalizationTarget.IsValid() && LocalizationTarget->Settings.SupportedCulturesStatistics.Num() > 0;
+}
+
+void FLocalizationTargetDetailCustomization::CountWords()
 {
 	if (LocalizationTarget.IsValid())
 	{
-		// Execute compile.
 		const TSharedPtr<SWindow> ParentWindow = FSlateApplication::Get().FindWidgetWindow(DetailLayoutBuilder->GetDetailsView().AsShared());
-		LocalizationCommandletTasks::CompileTarget(ParentWindow.ToSharedRef(), LocalizationTarget.Get());
+		LocalizationCommandletTasks::GenerateWordCountReportForTarget(ParentWindow.ToSharedRef(), LocalizationTarget.Get());
+
+		UpdateTargetFromReports();
 	}
 }
 
-void FLocalizationTargetDetailCustomization::RefreshWordCounts()
+bool FLocalizationTargetDetailCustomization::CanCompileAllCultures() const
+{
+	return LocalizationTarget.IsValid() && LocalizationTarget->Settings.SupportedCulturesStatistics.Num() > 0;
+}
+
+void FLocalizationTargetDetailCustomization::CompileAllCultures()
 {
 	if (LocalizationTarget.IsValid())
 	{
+		// Execute compile.
 		const TSharedPtr<SWindow> ParentWindow = FSlateApplication::Get().FindWidgetWindow(DetailLayoutBuilder->GetDetailsView().AsShared());
-		LocalizationCommandletTasks::GenerateWordCountReportForTarget(ParentWindow.ToSharedRef(), LocalizationTarget.Get());
-
-		UpdateTargetFromReports();
+		LocalizationCommandletTasks::CompileTarget(ParentWindow.ToSharedRef(), LocalizationTarget.Get());
 	}
 }
 
@@ -833,12 +869,26 @@ void FLocalizationTargetDetailCustomization::BuildListedCulturesList()
 
 		FString LeftName;
 		LeftNameHandle->GetValue(LeftName);
+		const FCulturePtr LeftCulture = FInternationalization::Get().GetCulture(LeftName);
 		FString RightName;
 		RightNameHandle->GetValue(RightName);
+		const FCulturePtr RightCulture = FInternationalization::Get().GetCulture(RightName);
 
-		return LeftName < RightName;
+		return LeftCulture.IsValid() && RightCulture.IsValid() ? LeftCulture->GetDisplayName() < RightCulture->GetDisplayName() : LeftName < RightName;
 	};
 	ListedCultureStatisticProperties.Sort(CultureSorter);
+
+	if (NoSupportedCulturesErrorText.IsValid())
+	{
+		if (ListedCultureStatisticProperties.Num())
+		{
+			NoSupportedCulturesErrorText->SetError(FText::GetEmpty());
+		}
+		else
+		{
+			NoSupportedCulturesErrorText->SetError(LOCTEXT("NoSupportedCulturesError", "At least one supported culture must be specified."));
+		}
+	}
 }
 
 void FLocalizationTargetDetailCustomization::RebuildListedCulturesList()
diff --git a/Engine/Source/Editor/LocalizationDashboard/Private/LocalizationTargetDetailCustomization.h b/Engine/Source/Editor/LocalizationDashboard/Private/LocalizationTargetDetailCustomization.h
--- a/Engine/Source/Editor/LocalizationDashboard/Private/LocalizationTargetDetailCustomization.h
+++ b/Engine/Source/Editor/LocalizationDashboard/Private/LocalizationTargetDetailCustomization.h
@@ -33,11 +33,16 @@ private:
 	void OnTargetDependencyCheckStateChanged(ULocalizationTarget* const OtherLocalizationTarget, const ECheckBoxState State);
 	ECheckBoxState IsTargetDependencyChecked(ULocalizationTarget* const OtherLocalizationTarget) const;
 
+	bool CanGather() const;
 	void Gather();
+	bool CanImportAllCultures() const;
 	void ImportAllCultures();
+	bool CanExportAllCultures() const;
 	void ExportAllCultures();
-	void Compile();
-	void RefreshWordCounts();
+	bool CanCountWords() const;
+	void CountWords();
+	bool CanCompileAllCultures() const;
+	void CompileAllCultures();
 	void UpdateTargetFromReports();
 
 	void BuildListedCulturesList();
@@ -68,6 +73,7 @@ private:
 	TSharedPtr<IPropertyHandle> SupportedCulturesStatisticsPropertyHandle;
 	FSimpleDelegate SupportedCulturesStatisticsPropertyHandle_OnNumElementsChanged;
 	TSharedPtr< SListView< TSharedPtr<IPropertyHandle> > > SupportedCultureListView;
+	TSharedPtr<SErrorText> NoSupportedCulturesErrorText;
 	TSharedPtr<SComboButton> AddNewSupportedCultureComboButton;
 	TSharedPtr<SCulturePicker> SupportedCulturePicker;
 	TArray< TSharedPtr<IPropertyHandle> > ListedCultureStatisticProperties;
diff --git a/Engine/Source/Editor/LocalizationDashboard/Private/LocalizationTargetSetDetailCustomization.cpp b/Engine/Source/Editor/LocalizationDashboard/Private/LocalizationTargetSetDetailCustomization.cpp
--- a/Engine/Source/Editor/LocalizationDashboard/Private/LocalizationTargetSetDetailCustomization.cpp
+++ b/Engine/Source/Editor/LocalizationDashboard/Private/LocalizationTargetSetDetailCustomization.cpp
@@ -56,6 +56,7 @@ namespace
 		TSharedPtr<FUICommandInfo> GatherAllTargets;
 		TSharedPtr<FUICommandInfo> ImportAllTargets;
 		TSharedPtr<FUICommandInfo> ExportAllTargets;
+		TSharedPtr<FUICommandInfo> CountWordsForAllTargets;
 		TSharedPtr<FUICommandInfo> CompileAllTargets;
 
 		/** Initialize commands */
@@ -67,6 +68,7 @@ namespace
 		UI_COMMAND( GatherAllTargets, "Gather All", "Gathers text for all targets in the project.", EUserInterfaceActionType::Button, FInputChord() );
 		UI_COMMAND( ImportAllTargets, "Import All", "Imports translations for all cultures of all targets in the project.", EUserInterfaceActionType::Button, FInputChord() );
 		UI_COMMAND( ExportAllTargets, "Export All", "Exports translations for all cultures of all targets in the project.", EUserInterfaceActionType::Button, FInputChord() );
+		UI_COMMAND( CountWordsForAllTargets, "Count Words All", "Counts translations for all targets in the project.", EUserInterfaceActionType::Button, FInputChord() );
 		UI_COMMAND( CompileAllTargets, "Compile All", "Compiles translations for all targets in the project.", EUserInterfaceActionType::Button, FInputChord() );
 	}
 }
@@ -95,17 +97,24 @@ void FLocalizationTargetSetDetailCustomization::CustomizeDetails(IDetailLayoutBu
 			const TSharedRef< FUICommandList > CommandList = MakeShareable(new FUICommandList);
 			FToolBarBuilder ToolBarBuilder( CommandList, FMultiBoxCustomization::AllowCustomization("LocalizationDashboard") );
 
-			CommandList->MapAction( FLocalizationDashboardCommands::Get().GatherAllTargets, FExecuteAction::CreateSP(this, &FLocalizationTargetSetDetailCustomization::GatherAllTargets));
-			ToolBarBuilder.AddToolBarButton(FLocalizationDashboardCommands::Get().GatherAllTargets, NAME_None, TAttribute<FText>(), TAttribute<FText>(), FSlateIcon(FEditorStyle::GetStyleSetName(), "LocalizationDashboard.GatherAllTargets"));
+			TAttribute<FText> GatherAllTargetsToolTipTextAttribute = TAttribute<FText>::Create(TAttribute<FText>::FGetter::CreateLambda([this]() -> FText
+				{
+					return CanGatherAllTargets() ? FLocalizationDashboardCommands::Get().GatherAllTargets->GetDescription() : LOCTEXT("GatherAllTargetsDisabledToolTip", "At least one target must have a native culture specified in order to gather.");
+				}));
+			CommandList->MapAction( FLocalizationDashboardCommands::Get().GatherAllTargets, FExecuteAction::CreateSP(this, &FLocalizationTargetSetDetailCustomization::GatherAllTargets), FCanExecuteAction::CreateSP(this, &FLocalizationTargetSetDetailCustomization::CanGatherAllTargets));
+			ToolBarBuilder.AddToolBarButton(FLocalizationDashboardCommands::Get().GatherAllTargets, NAME_None, TAttribute<FText>(), GatherAllTargetsToolTipTextAttribute, FSlateIcon(FEditorStyle::GetStyleSetName(), "LocalizationDashboard.GatherAllTargets"));
+
+			CommandList->MapAction( FLocalizationDashboardCommands::Get().ImportAllTargets, FExecuteAction::CreateSP(this, &FLocalizationTargetSetDetailCustomization::ImportAllTargets), FCanExecuteAction::CreateSP(this, &FLocalizationTargetSetDetailCustomization::CanImportAllTargets));
+			ToolBarBuilder.AddToolBarButton(FLocalizationDashboardCommands::Get().ImportAllTargets, NAME_None, TAttribute<FText>(), TAttribute<FText>(), FSlateIcon(FEditorStyle::GetStyleSetName(), "LocalizationDashboard.ImportAllTargetsAllCultures"));
 
-			CommandList->MapAction( FLocalizationDashboardCommands::Get().ImportAllTargets, FExecuteAction::CreateSP(this, &FLocalizationTargetSetDetailCustomization::ImportAllTargets));
-			ToolBarBuilder.AddToolBarButton(FLocalizationDashboardCommands::Get().ImportAllTargets, NAME_None, TAttribute<FText>(), TAttribute<FText>(), FSlateIcon(FEditorStyle::GetStyleSetName(), "LocalizationDashboard.ImportForAllTargetsCultures"));
+			CommandList->MapAction( FLocalizationDashboardCommands::Get().ExportAllTargets, FExecuteAction::CreateSP(this, &FLocalizationTargetSetDetailCustomization::ExportAllTargets), FCanExecuteAction::CreateSP(this, &FLocalizationTargetSetDetailCustomization::CanExportAllTargets));
+			ToolBarBuilder.AddToolBarButton(FLocalizationDashboardCommands::Get().ExportAllTargets, NAME_None, TAttribute<FText>(), TAttribute<FText>(), FSlateIcon(FEditorStyle::GetStyleSetName(), "LocalizationDashboard.ExportAllTargetsAllCultures"));
 
-			CommandList->MapAction( FLocalizationDashboardCommands::Get().ExportAllTargets, FExecuteAction::CreateSP(this, &FLocalizationTargetSetDetailCustomization::ExportAllTargets));
-			ToolBarBuilder.AddToolBarButton(FLocalizationDashboardCommands::Get().ExportAllTargets, NAME_None, TAttribute<FText>(), TAttribute<FText>(), FSlateIcon(FEditorStyle::GetStyleSetName(), "LocalizationDashboard.ExportForAllTargetsCultures"));
+			CommandList->MapAction( FLocalizationDashboardCommands::Get().CountWordsForAllTargets, FExecuteAction::CreateSP(this, &FLocalizationTargetSetDetailCustomization::CountWordsForAllTargets), FCanExecuteAction::CreateSP(this, &FLocalizationTargetSetDetailCustomization::CanCountWordsForAllTargets));
+			ToolBarBuilder.AddToolBarButton(FLocalizationDashboardCommands::Get().CountWordsForAllTargets, NAME_None, TAttribute<FText>(), TAttribute<FText>(), FSlateIcon(FEditorStyle::GetStyleSetName(), "LocalizationDashboard.CountWordsForAllTargets"));
 
-			CommandList->MapAction( FLocalizationDashboardCommands::Get().CompileAllTargets, FExecuteAction::CreateSP(this, &FLocalizationTargetSetDetailCustomization::CompileAllTargets));
-			ToolBarBuilder.AddToolBarButton(FLocalizationDashboardCommands::Get().CompileAllTargets, NAME_None, TAttribute<FText>(), TAttribute<FText>(), FSlateIcon(FEditorStyle::GetStyleSetName(), "LocalizationDashboard.CompileAllTargets"));
+			CommandList->MapAction( FLocalizationDashboardCommands::Get().CompileAllTargets, FExecuteAction::CreateSP(this, &FLocalizationTargetSetDetailCustomization::CompileAllTargets), FCanExecuteAction::CreateSP(this, &FLocalizationTargetSetDetailCustomization::CanCompileAllTargets));
+			ToolBarBuilder.AddToolBarButton(FLocalizationDashboardCommands::Get().CompileAllTargets, NAME_None, TAttribute<FText>(), TAttribute<FText>(), FSlateIcon(FEditorStyle::GetStyleSetName(), "LocalizationDashboard.CompileAllTargetsAllCultures"));
 
 			BuildTargetsList();
 
@@ -254,47 +263,80 @@ void FLocalizationTargetSetDetailCustomization::RebuildTargetsList()
 //	DetailLayoutBuilder->ForceRefreshDetails();
 //}
 
-void FLocalizationTargetSetDetailCustomization::GatherAllTargets()
+bool FLocalizationTargetSetDetailCustomization::CanGatherAllTargets() const
 {
 	for (ULocalizationTarget* const LocalizationTarget : TargetSet->TargetObjects)
 	{
-		// Save unsaved packages.
-		const bool bPromptUserToSave = true;
-		const bool bSaveMapPackages = true;
-		const bool bSaveContentPackages = true;
-		const bool bFastSave = false;
-		const bool bNotifyNoPackagesSaved = false;
-		const bool bCanBeDeclined = true;
-		bool DidPackagesNeedSaving;
-		const bool WerePackagesSaved = FEditorFileUtils::SaveDirtyPackages(bPromptUserToSave, bSaveMapPackages, bSaveContentPackages, bFastSave, bNotifyNoPackagesSaved, bCanBeDeclined, &DidPackagesNeedSaving);
-
-		if (DidPackagesNeedSaving && !WerePackagesSaved)
+		// If any target can be gathered, then gathering all can be executed.
+		if (LocalizationTarget && LocalizationTarget->Settings.SupportedCulturesStatistics.Num() > 0 && LocalizationTarget->Settings.SupportedCulturesStatistics.IsValidIndex(LocalizationTarget->Settings.NativeCultureIndex))
 		{
-			// Give warning dialog.
-			const FText MessageText = NSLOCTEXT("LocalizationCultureActions", "UnsavedPackagesWarningDialogMessage", "There are unsaved changes. These changes may not be gathered from correctly.");
-			const FText TitleText = NSLOCTEXT("LocalizationCultureActions", "UnsavedPackagesWarningDialogTitle", "Unsaved Changes Before Gather");
-			switch(FMessageDialog::Open(EAppMsgType::OkCancel, MessageText, &TitleText))
+			return true;
+		}
+	}
+
+	return false;
+}
+
+void FLocalizationTargetSetDetailCustomization::GatherAllTargets()
+{
+	// Save unsaved packages.
+	const bool bPromptUserToSave = true;
+	const bool bSaveMapPackages = true;
+	const bool bSaveContentPackages = true;
+	const bool bFastSave = false;
+	const bool bNotifyNoPackagesSaved = false;
+	const bool bCanBeDeclined = true;
+	bool DidPackagesNeedSaving;
+	const bool WerePackagesSaved = FEditorFileUtils::SaveDirtyPackages(bPromptUserToSave, bSaveMapPackages, bSaveContentPackages, bFastSave, bNotifyNoPackagesSaved, bCanBeDeclined, &DidPackagesNeedSaving);
+
+	if (DidPackagesNeedSaving && !WerePackagesSaved)
+	{
+		// Give warning dialog.
+		const FText MessageText = NSLOCTEXT("LocalizationCultureActions", "UnsavedPackagesWarningDialogMessage", "There are unsaved changes. These changes may not be gathered from correctly.");
+		const FText TitleText = NSLOCTEXT("LocalizationCultureActions", "UnsavedPackagesWarningDialogTitle", "Unsaved Changes Before Gather");
+		switch(FMessageDialog::Open(EAppMsgType::OkCancel, MessageText, &TitleText))
+		{
+		case EAppReturnType::Cancel:
 			{
-			case EAppReturnType::Cancel:
-				{
-					return;
-				}
-				break;
+				return;
 			}
+			break;
 		}
+	}
 
+	TArray<ULocalizationTarget*> TargetObjectsToProcess;
+	for (ULocalizationTarget* const LocalizationTarget : TargetSet->TargetObjects)
+	{
+		if (LocalizationTarget && LocalizationTarget->Settings.SupportedCulturesStatistics.Num() > 0 && LocalizationTarget->Settings.SupportedCulturesStatistics.IsValidIndex(LocalizationTarget->Settings.NativeCultureIndex))
+		{
+			TargetObjectsToProcess.Add(LocalizationTarget);
+		}
 	}
 
 	// Execute gather.
 	const TSharedPtr<SWindow> ParentWindow = FSlateApplication::Get().FindWidgetWindow(DetailLayoutBuilder->GetDetailsView().AsShared());
-	LocalizationCommandletTasks::GatherTargets(ParentWindow.ToSharedRef(), TargetSet->TargetObjects);
+	LocalizationCommandletTasks::GatherTargets(ParentWindow.ToSharedRef(), TargetObjectsToProcess);
 
-	for (ULocalizationTarget* const LocalizationTarget : TargetSet->TargetObjects)
+	for (ULocalizationTarget* const LocalizationTarget : TargetObjectsToProcess)
 	{
 		UpdateTargetFromReports(LocalizationTarget);
 	}
 }
 
+bool FLocalizationTargetSetDetailCustomization::CanImportAllTargets() const
+{
+	for (ULocalizationTarget* const LocalizationTarget : TargetSet->TargetObjects)
+	{
+		// If any target can be imported, then importing all can be executed.
+		if (LocalizationTarget && LocalizationTarget->Settings.SupportedCulturesStatistics.Num() > 0)
+		{
+			return true;
+		}
+	}
+
+	return false;
+}
+
 void FLocalizationTargetSetDetailCustomization::ImportAllTargets()
 {
 	IDesktopPlatform* DesktopPlatform = FDesktopPlatformModule::Get();
@@ -314,9 +356,18 @@ void FLocalizationTargetSetDetailCustomization::ImportAllTargets()
 		FString OutputDirectory;
 		if (DesktopPlatform->OpenDirectoryDialog(ParentWindowWindowHandle, LOCTEXT("ImportAllTranslationsDialogTitle", "Import All Translations from Directory").ToString(), DefaultPath, OutputDirectory))
 		{
-			LocalizationCommandletTasks::ImportTargets(ParentWindow.ToSharedRef(), TargetSet->TargetObjects, TOptional<FString>(OutputDirectory));
-
+			TArray<ULocalizationTarget*> TargetObjectsToProcess;
 			for (ULocalizationTarget* const LocalizationTarget : TargetSet->TargetObjects)
+			{
+				if (LocalizationTarget && LocalizationTarget->Settings.SupportedCulturesStatistics.Num() > 0)
+				{
+					TargetObjectsToProcess.Add(LocalizationTarget);
+				}
+			}
+
+			LocalizationCommandletTasks::ImportTargets(ParentWindow.ToSharedRef(), TargetObjectsToProcess, TOptional<FString>(OutputDirectory));
+
+			for (ULocalizationTarget* const LocalizationTarget : TargetObjectsToProcess)
 			{
 				UpdateTargetFromReports(LocalizationTarget);
 			}
@@ -324,6 +375,20 @@ void FLocalizationTargetSetDetailCustomization::ImportAllTargets()
 	}
 }
 
+bool FLocalizationTargetSetDetailCustomization::CanExportAllTargets() const
+{
+	for (ULocalizationTarget* const LocalizationTarget : TargetSet->TargetObjects)
+	{
+		// If any target can be imported, then exporting all can be executed.
+		if (LocalizationTarget && LocalizationTarget->Settings.SupportedCulturesStatistics.Num() > 0)
+		{
+			return true;
+		}
+	}
+
+	return false;
+}
+
 void FLocalizationTargetSetDetailCustomization::ExportAllTargets()
 {
 	IDesktopPlatform* DesktopPlatform = FDesktopPlatformModule::Get();
@@ -342,9 +407,78 @@ void FLocalizationTargetSetDetailCustomization::ExportAllTargets()
 		FString OutputDirectory;
 		if (DesktopPlatform->OpenDirectoryDialog(ParentWindowWindowHandle, LOCTEXT("ExportAllTranslationsDialogTitle", "Export All Translations to Directory").ToString(), DefaultPath, OutputDirectory))
 		{
-			LocalizationCommandletTasks::ExportTargets(ParentWindow.ToSharedRef(), TargetSet->TargetObjects, TOptional<FString>(OutputDirectory));
+			TArray<ULocalizationTarget*> TargetObjectsToProcess;
+			for (ULocalizationTarget* const LocalizationTarget : TargetSet->TargetObjects)
+			{
+				if (LocalizationTarget && LocalizationTarget->Settings.SupportedCulturesStatistics.Num() > 0)
+				{
+					TargetObjectsToProcess.Add(LocalizationTarget);
+				}
+			}
+
+			LocalizationCommandletTasks::ExportTargets(ParentWindow.ToSharedRef(), TargetObjectsToProcess, TOptional<FString>(OutputDirectory));
+		}
+	}
+}
+
+bool FLocalizationTargetSetDetailCustomization::CanCountWordsForAllTargets() const
+{
+	for (ULocalizationTarget* const LocalizationTarget : TargetSet->TargetObjects)
+	{
+		// If any target can be imported, then counting words for all can be executed.
+		if (LocalizationTarget && LocalizationTarget->Settings.SupportedCulturesStatistics.Num() > 0)
+		{
+			return true;
 		}
 	}
+
+	return false;
+}
+
+void FLocalizationTargetSetDetailCustomization::CountWordsForAllTargets()
+{
+	TArray<ULocalizationTarget*> TargetObjectsToProcess;
+	for (ULocalizationTarget* const LocalizationTarget : TargetSet->TargetObjects)
+	{
+		if (LocalizationTarget && LocalizationTarget->Settings.SupportedCulturesStatistics.Num() > 0)
+		{
+			TargetObjectsToProcess.Add(LocalizationTarget);
+		}
+	}
+
+	// Execute compile.
+	const TSharedPtr<SWindow> ParentWindow = FSlateApplication::Get().FindWidgetWindow(DetailLayoutBuilder->GetDetailsView().AsShared());
+	LocalizationCommandletTasks::GenerateWordCountReportsForTargets(ParentWindow.ToSharedRef(), TargetObjectsToProcess);
+}
+
+bool FLocalizationTargetSetDetailCustomization::CanCompileAllTargets() const
+{
+	for (ULocalizationTarget* const LocalizationTarget : TargetSet->TargetObjects)
+	{
+		// If any target can be imported, then compiling all can be executed.
+		if (LocalizationTarget && LocalizationTarget->Settings.SupportedCulturesStatistics.Num() > 0)
+		{
+			return true;
+		}
+	}
+
+	return false;
+}
+
+void FLocalizationTargetSetDetailCustomization::CompileAllTargets()
+{
+	TArray<ULocalizationTarget*> TargetObjectsToProcess;
+	for (ULocalizationTarget* const LocalizationTarget : TargetSet->TargetObjects)
+	{
+		if (LocalizationTarget && LocalizationTarget->Settings.SupportedCulturesStatistics.Num() > 0)
+		{
+			TargetObjectsToProcess.Add(LocalizationTarget);
+		}
+	}
+
+	// Execute compile.
+	const TSharedPtr<SWindow> ParentWindow = FSlateApplication::Get().FindWidgetWindow(DetailLayoutBuilder->GetDetailsView().AsShared());
+	LocalizationCommandletTasks::CompileTargets(ParentWindow.ToSharedRef(), TargetObjectsToProcess);
 }
 
 void FLocalizationTargetSetDetailCustomization::UpdateTargetFromReports(ULocalizationTarget* const LocalizationTarget)
@@ -396,13 +530,6 @@ void FLocalizationTargetSetDetailCustomization::UpdateTargetFromReports(ULocaliz
 	//}
 }
 
-void FLocalizationTargetSetDetailCustomization::CompileAllTargets()
-{
-	// Execute compile.
-	const TSharedPtr<SWindow> ParentWindow = FSlateApplication::Get().FindWidgetWindow(DetailLayoutBuilder->GetDetailsView().AsShared());
-	LocalizationCommandletTasks::CompileTargets(ParentWindow.ToSharedRef(), TargetSet->TargetObjects);
-}
-
 TSharedRef<ITableRow> FLocalizationTargetSetDetailCustomization::OnGenerateRow(TSharedPtr<IPropertyHandle> TargetObjectPropertyHandle, const TSharedRef<STableViewBase>& Table)
 {
 	return SNew(SLocalizationDashboardTargetRow, Table, DetailLayoutBuilder->GetPropertyUtilities(), TargetObjectPropertyHandle.ToSharedRef());
diff --git a/Engine/Source/Editor/LocalizationDashboard/Private/LocalizationTargetSetDetailCustomization.h b/Engine/Source/Editor/LocalizationDashboard/Private/LocalizationTargetSetDetailCustomization.h
--- a/Engine/Source/Editor/LocalizationDashboard/Private/LocalizationTargetSetDetailCustomization.h
+++ b/Engine/Source/Editor/LocalizationDashboard/Private/LocalizationTargetSetDetailCustomization.h
@@ -32,11 +32,17 @@ private:
 	//TSharedRef<SWidget> ServiceProviderComboBox_OnGenerateWidget(TSharedPtr<FLocalizationServiceProviderWrapper> LSPWrapper) const;
 	//void ServiceProviderComboBox_OnSelectionChanged(TSharedPtr<FLocalizationServiceProviderWrapper> LSPWrapper, ESelectInfo::Type SelectInfo);
 
+	bool CanGatherAllTargets() const;
 	void GatherAllTargets();
+	bool CanImportAllTargets() const;
 	void ImportAllTargets();
+	bool CanExportAllTargets() const;
 	void ExportAllTargets();
-	void UpdateTargetFromReports(ULocalizationTarget* const LocalizationTarget);
+	bool CanCountWordsForAllTargets() const;
+	void CountWordsForAllTargets();
+	bool CanCompileAllTargets() const;
 	void CompileAllTargets();
+	void UpdateTargetFromReports(ULocalizationTarget* const LocalizationTarget);
 
 	TSharedRef<ITableRow> OnGenerateRow(TSharedPtr<IPropertyHandle> TargetObjectPropertyHandle, const TSharedRef<STableViewBase>& Table);
 	FReply OnNewTargetButtonClicked();
diff --git a/Engine/Source/Editor/LocalizationDashboard/Private/SLocalizationDashboardTargetRow.cpp b/Engine/Source/Editor/LocalizationDashboard/Private/SLocalizationDashboardTargetRow.cpp
--- a/Engine/Source/Editor/LocalizationDashboard/Private/SLocalizationDashboardTargetRow.cpp
+++ b/Engine/Source/Editor/LocalizationDashboard/Private/SLocalizationDashboardTargetRow.cpp
@@ -56,6 +56,14 @@ TSharedRef<SWidget> SLocalizationDashboardTargetRow::GenerateWidgetForColumn( co
 		TSharedRef<SHorizontalBox> HorizontalBox = SNew(SHorizontalBox);
 		Return = HorizontalBox;
 
+		const TSharedRef<SLocalizationDashboardTargetRow> LocalSharedThis = SharedThis(this);
+
+		const auto& CanGather = [=]() -> bool
+		{
+			ULocalizationTarget* const LocalizationTarget = LocalSharedThis->GetTarget();
+			return LocalizationTarget && LocalizationTarget->Settings.SupportedCulturesStatistics.Num() > 0 && LocalizationTarget->Settings.SupportedCulturesStatistics.IsValidIndex(LocalizationTarget->Settings.NativeCultureIndex);
+		};
+
 		// Gather
 		HorizontalBox->AddSlot()
 			.FillWidth(1.0f)
@@ -64,8 +72,12 @@ TSharedRef<SWidget> SLocalizationDashboardTargetRow::GenerateWidgetForColumn( co
 			[
 				SNew(SButton)
 				.ButtonStyle( FEditorStyle::Get(), TEXT("HoverHintOnly") )
-				.ToolTipText( LOCTEXT("GatherButtonLabel", "Gather") )
+				.ToolTipText_Lambda([=]() -> FText
+				{
+					return CanGather() ? LOCTEXT("GatherButtonLabel", "Gather") : LOCTEXT("DisabledGatherButtonLabel", "Must have a native culture specified in order to gather.");
+				})
 				.OnClicked(this, &SLocalizationDashboardTargetRow::Gather)
+				.IsEnabled_Lambda(CanGather)
 				.Content()
 				[
 					SNew(SImage)
@@ -73,7 +85,7 @@ TSharedRef<SWidget> SLocalizationDashboardTargetRow::GenerateWidgetForColumn( co
 				]
 			];
 
-		// Refresh Word Counts
+		// Import All
 		HorizontalBox->AddSlot()
 			.FillWidth(1.0f)
 			.HAlign(HAlign_Center)
@@ -81,16 +93,21 @@ TSharedRef<SWidget> SLocalizationDashboardTargetRow::GenerateWidgetForColumn( co
 			[
 				SNew(SButton)
 				.ButtonStyle( FEditorStyle::Get(), TEXT("HoverHintOnly") )
-				.ToolTipText( LOCTEXT("RefreshWordCountButtonLabel", "Refresh") )
-				.OnClicked(this, &SLocalizationDashboardTargetRow::RefreshWordCount)
+				.ToolTipText( LOCTEXT("ImportAllButtonLabel", "Import All") )
+				.OnClicked(this, &SLocalizationDashboardTargetRow::ImportAll)
+				.IsEnabled_Lambda([this]() -> bool
+				{
+					ULocalizationTarget* const LocalizationTarget = GetTarget();
+					return LocalizationTarget && LocalizationTarget->Settings.SupportedCulturesStatistics.Num() > 0;
+				})
 				.Content()
 				[
 					SNew(SImage)
-					.Image(FEditorStyle::GetBrush("LocalizationDashboard.RefreshWordCount"))
+					.Image( FEditorStyle::GetBrush("LocalizationDashboard.ImportAllCultures") )
 				]
 			];
 
-		// Import All
+		// Export All
 		HorizontalBox->AddSlot()
 			.FillWidth(1.0f)
 			.HAlign(HAlign_Center)
@@ -98,16 +115,21 @@ TSharedRef<SWidget> SLocalizationDashboardTargetRow::GenerateWidgetForColumn( co
 			[
 				SNew(SButton)
 				.ButtonStyle( FEditorStyle::Get(), TEXT("HoverHintOnly") )
-				.ToolTipText( LOCTEXT("ImportAllButtonLabel", "Import All") )
-				.OnClicked(this, &SLocalizationDashboardTargetRow::ImportAll)
+				.ToolTipText(LOCTEXT("ExportAllButtonLabel", "Export All"))
+				.OnClicked(this, &SLocalizationDashboardTargetRow::ExportAll)
+				.IsEnabled_Lambda([this]() -> bool
+				{
+					ULocalizationTarget* const LocalizationTarget = GetTarget();
+					return LocalizationTarget && LocalizationTarget->Settings.SupportedCulturesStatistics.Num() > 0;
+				})
 				.Content()
 				[
 					SNew(SImage)
-					.Image( FEditorStyle::GetBrush("LocalizationDashboard.ImportForAllCultures") )
+					.Image(FEditorStyle::GetBrush("LocalizationDashboard.ExportAllCultures"))
 				]
 			];
 
-		// Export All
+		// Count Words
 		HorizontalBox->AddSlot()
 			.FillWidth(1.0f)
 			.HAlign(HAlign_Center)
@@ -115,12 +137,17 @@ TSharedRef<SWidget> SLocalizationDashboardTargetRow::GenerateWidgetForColumn( co
 			[
 				SNew(SButton)
 				.ButtonStyle( FEditorStyle::Get(), TEXT("HoverHintOnly") )
-				.ToolTipText(LOCTEXT("ExportAllButtonLabel", "Export All"))
-				.OnClicked(this, &SLocalizationDashboardTargetRow::ExportAll)
+				.ToolTipText( LOCTEXT("RefreshWordCountButtonLabel", "Count Words for All") )
+				.OnClicked(this, &SLocalizationDashboardTargetRow::CountWords)
+				.IsEnabled_Lambda([this]() -> bool
+				{
+					ULocalizationTarget* const LocalizationTarget = GetTarget();
+					return LocalizationTarget && LocalizationTarget->Settings.SupportedCulturesStatistics.Num() > 0;
+				})
 				.Content()
 				[
 					SNew(SImage)
-					.Image(FEditorStyle::GetBrush("LocalizationDashboard.ExportForAllCultures"))
+					.Image(FEditorStyle::GetBrush("LocalizationDashboard.CountWordsForTarget"))
 				]
 			];
 
@@ -132,12 +159,17 @@ TSharedRef<SWidget> SLocalizationDashboardTargetRow::GenerateWidgetForColumn( co
 			[
 				SNew(SButton)
 				.ButtonStyle( FEditorStyle::Get(), TEXT("HoverHintOnly") )
-				.ToolTipText( LOCTEXT("CompileButtonLabel", "Compile") )
-				.OnClicked(this, &SLocalizationDashboardTargetRow::Compile)
+				.ToolTipText( LOCTEXT("CompileButtonLabel", "Compile All") )
+				.OnClicked(this, &SLocalizationDashboardTargetRow::CompileAll)
+				.IsEnabled_Lambda([this]() -> bool
+				{
+					ULocalizationTarget* const LocalizationTarget = GetTarget();
+					return LocalizationTarget && LocalizationTarget->Settings.SupportedCulturesStatistics.Num() > 0;
+				})
 				.Content()
 				[
 					SNew(SImage)
-					.Image( FEditorStyle::GetBrush("LocalizationDashboard.CompileTarget") )
+					.Image( FEditorStyle::GetBrush("LocalizationDashboard.CompileAllCultures") )
 				]
 			];
 
@@ -305,22 +337,6 @@ void SLocalizationDashboardTargetRow::UpdateTargetFromReports()
 	}
 }
 
-FReply SLocalizationDashboardTargetRow::RefreshWordCount()
-{
-	bool HasSucceeded = false;
-
-	ULocalizationTarget* const LocalizationTarget = GetTarget();
-	if (LocalizationTarget)
-	{
-		const TSharedPtr<SWindow> ParentWindow = FSlateApplication::Get().FindWidgetWindow(AsShared());
-		LocalizationCommandletTasks::GenerateWordCountReportForTarget(ParentWindow.ToSharedRef(), LocalizationTarget);
-
-		UpdateTargetFromReports();
-	}
-
-	return FReply::Handled();
-}
-
 FReply SLocalizationDashboardTargetRow::Gather()
 {
 	ULocalizationTarget* const LocalizationTarget = GetTarget();
@@ -428,7 +444,23 @@ FReply SLocalizationDashboardTargetRow::ExportAll()
 	return FReply::Handled();
 }
 
-FReply SLocalizationDashboardTargetRow::Compile()
+FReply SLocalizationDashboardTargetRow::CountWords()
+{
+	bool HasSucceeded = false;
+
+	ULocalizationTarget* const LocalizationTarget = GetTarget();
+	if (LocalizationTarget)
+	{
+		const TSharedPtr<SWindow> ParentWindow = FSlateApplication::Get().FindWidgetWindow(AsShared());
+		LocalizationCommandletTasks::GenerateWordCountReportForTarget(ParentWindow.ToSharedRef(), LocalizationTarget);
+
+		UpdateTargetFromReports();
+	}
+
+	return FReply::Handled();
+}
+
+FReply SLocalizationDashboardTargetRow::CompileAll()
 {
 	ULocalizationTarget* const LocalizationTarget = GetTarget();
 	if (LocalizationTarget)
diff --git a/Engine/Source/Editor/LocalizationDashboard/Private/SLocalizationDashboardTargetRow.h b/Engine/Source/Editor/LocalizationDashboard/Private/SLocalizationDashboardTargetRow.h
--- a/Engine/Source/Editor/LocalizationDashboard/Private/SLocalizationDashboardTargetRow.h
+++ b/Engine/Source/Editor/LocalizationDashboard/Private/SLocalizationDashboardTargetRow.h
@@ -29,12 +29,11 @@ private:
 
 	void UpdateTargetFromReports();
 
-	FReply RefreshWordCount();
-
 	FReply Gather();
 	FReply ImportAll();
 	FReply ExportAll();
-	FReply Compile();
+	FReply CountWords();
+	FReply CompileAll();
 
 	FReply EnqueueDeletion();
 	void Delete();
diff --git a/Engine/Source/Editor/LocalizationDashboard/Private/SLocalizationTargetEditorCultureRow.cpp b/Engine/Source/Editor/LocalizationDashboard/Private/SLocalizationTargetEditorCultureRow.cpp
--- a/Engine/Source/Editor/LocalizationDashboard/Private/SLocalizationTargetEditorCultureRow.cpp
+++ b/Engine/Source/Editor/LocalizationDashboard/Private/SLocalizationTargetEditorCultureRow.cpp
@@ -90,7 +90,7 @@ TSharedRef<SWidget> SLocalizationTargetEditorCultureRow::GenerateWidgetForColumn
 				.Content()
 				[
 					SNew(SImage)
-					.Image( FEditorStyle::GetBrush("LocalizationTargetEditor.ImportForCulture") )
+					.Image( FEditorStyle::GetBrush("LocalizationTargetEditor.ImportCulture") )
 				]
 			];
 
@@ -107,7 +107,24 @@ TSharedRef<SWidget> SLocalizationTargetEditorCultureRow::GenerateWidgetForColumn
 				.Content()
 				[
 					SNew(SImage)
-					.Image(FEditorStyle::GetBrush("LocalizationTargetEditor.ExportForCulture"))
+					.Image(FEditorStyle::GetBrush("LocalizationTargetEditor.ExportCulture"))
+				]
+			];
+
+		// Export
+		HorizontalBox->AddSlot()
+			.FillWidth(1.0f)
+			.HAlign(HAlign_Center)
+			.VAlign(VAlign_Center)
+			[
+				SNew(SButton)
+				.ButtonStyle( FEditorStyle::Get(), TEXT("HoverHintOnly") )
+				.ToolTipText(NSLOCTEXT("LocalizationTargetCultureActions", "CompileButtonLabel", "Compile"))
+				.OnClicked(this, &SLocalizationTargetEditorCultureRow::Compile)
+				.Content()
+				[
+					SNew(SImage)
+					.Image(FEditorStyle::GetBrush("LocalizationTargetEditor.CompileCulture"))
 				]
 			];
 
@@ -401,6 +418,19 @@ FReply SLocalizationTargetEditorCultureRow::Export()
 	return FReply::Handled();
 }
 
+FReply SLocalizationTargetEditorCultureRow::Compile()
+{
+	const FCulturePtr Culture = GetCulture();
+	ULocalizationTarget* const LocalizationTarget = GetTarget();
+	if (Culture.IsValid() && LocalizationTarget)
+	{
+		const TSharedPtr<SWindow> ParentWindow = FSlateApplication::Get().FindWidgetWindow(AsShared());
+		LocalizationCommandletTasks::CompileCulture(ParentWindow.ToSharedRef(), LocalizationTarget, Culture->GetName());
+	}
+
+	return FReply::Handled();
+}
+
 bool SLocalizationTargetEditorCultureRow::CanDelete() const
 {
 	return !IsNativeCultureForTarget();
diff --git a/Engine/Source/Editor/LocalizationDashboard/Private/SLocalizationTargetEditorCultureRow.h b/Engine/Source/Editor/LocalizationDashboard/Private/SLocalizationTargetEditorCultureRow.h
--- a/Engine/Source/Editor/LocalizationDashboard/Private/SLocalizationTargetEditorCultureRow.h
+++ b/Engine/Source/Editor/LocalizationDashboard/Private/SLocalizationTargetEditorCultureRow.h
@@ -34,6 +34,7 @@ private:
 	FReply Edit();
 	FReply Import();
 	FReply Export();
+	FReply Compile();
 	bool CanDelete() const;
 	FReply EnqueueDeletion();
 	void Delete();
diff --git a/Engine/Source/Editor/UnrealEd/Private/Commandlets/GenerateGatherArchiveCommandlet.cpp b/Engine/Source/Editor/UnrealEd/Private/Commandlets/GenerateGatherArchiveCommandlet.cpp
--- a/Engine/Source/Editor/UnrealEd/Private/Commandlets/GenerateGatherArchiveCommandlet.cpp
+++ b/Engine/Source/Editor/UnrealEd/Private/Commandlets/GenerateGatherArchiveCommandlet.cpp
@@ -66,15 +66,37 @@ int32 UGenerateGatherArchiveCommandlet::Main( const FString& Params )
 		return -1;
 	}
 
-	// Get source culture.
-	FString SourceCulture;
-	if( GetStringFromConfig( *SectionName, TEXT("SourceCulture"), SourceCulture, GatherTextConfigPath ) )
+	// Get native culture.
+	FString NativeCulture;
 	{
-		if( I18N.GetCulture( SourceCulture ).IsValid() )
+		const bool WasNativeCultureSpecified = GetStringFromConfig( *SectionName, TEXT("NativeCulture"), NativeCulture, GatherTextConfigPath );
+
+		// Get source culture (DEPRECATED).
+		FString SourceCulture;
+		const bool WasSourceCultureSpecified = GetStringFromConfig( *SectionName, TEXT("SourceCulture"), SourceCulture, GatherTextConfigPath );
+		if (WasSourceCultureSpecified)
+		{
+			UE_LOG(LogGatherTextFromSourceCommandlet, Warning, TEXT("SourceCulture detected in section %s. SourceCulture is deprecated, please use NativeCulture."), *SectionName);
+		}
+
+		// Use SourceCulture if NativeCulture isn't specified.
+		if (!WasNativeCultureSpecified && WasSourceCultureSpecified)
+		{
+			NativeCulture = SourceCulture;
+		}
+	}
+
+	if (!NativeCulture.IsEmpty())
+	{
+		if (!I18N.GetCulture(NativeCulture).IsValid())
 		{
-			UE_LOG(LogGenerateArchiveCommandlet, Verbose, TEXT("Specified culture is not a valid runtime culture, but may be a valid base language: %s"), *(SourceCulture) );
+			UE_LOG(LogGenerateArchiveCommandlet, Verbose, TEXT("Specified native culture is not a valid runtime culture, but may be a valid base language: %s"), *(NativeCulture) );
 		}
 	}
+	else
+	{
+		UE_LOG(LogGenerateArchiveCommandlet, Warning, TEXT("No native culture specified. If the actual native culture is specified for generation, its archive entries may lack default translations."));
+	}
 
 	// Get cultures to generate.
 	TArray<FString> CulturesToGenerate;
@@ -134,7 +156,7 @@ int32 UGenerateGatherArchiveCommandlet::Main( const FString& Params )
 	for(int32 Culture = 0; Culture < CulturesToGenerate.Num(); Culture++)
 	{
 		TSharedRef< FInternationalizationArchive > InternationalizationArchive = MakeShareable( new FInternationalizationArchive );
-		BuildArchiveFromManifest( InternationalizationManifest, InternationalizationArchive, SourceCulture, CulturesToGenerate[Culture] );
+		BuildArchiveFromManifest( InternationalizationManifest, InternationalizationArchive, NativeCulture, CulturesToGenerate[Culture] );
 
 		const FString CulturePath = DestinationPath / CulturesToGenerate[Culture];
 		FJsonInternationalizationArchiveSerializer ArchiveSerializer;
