Add Localization Service Provider settings to Localization Dashboard
Also deleting Localization Dashboard's ILocalizationServiceProvider, as this class now exists in the LocalizationService module.

diff --git a/Engine/Source/Developer/LocalizationService/Public/ILocalizationServiceProvider.h b/Engine/Source/Developer/LocalizationService/Public/ILocalizationServiceProvider.h
--- a/Engine/Source/Developer/LocalizationService/Public/ILocalizationServiceProvider.h
+++ b/Engine/Source/Developer/LocalizationService/Public/ILocalizationServiceProvider.h
@@ -9,7 +9,7 @@
 class IDetailCategoryBuilder;
 class ULocalizationTarget;
 
-#define LOCALIZATION_SERVICES_WITH_SLATE			!PLATFORM_LINUX
+#define LOCALIZATION_SERVICES_WITH_SLATE	(!(PLATFORM_LINUX && IS_PROGRAM))
 
 /**
  * Hint for how to execute the operation. Note that asynchronous operations require
diff --git a/Engine/Source/Editor/LocalizationDashboard/LocalizationDashboard.Build.cs b/Engine/Source/Editor/LocalizationDashboard/LocalizationDashboard.Build.cs
--- a/Engine/Source/Editor/LocalizationDashboard/LocalizationDashboard.Build.cs
+++ b/Engine/Source/Editor/LocalizationDashboard/LocalizationDashboard.Build.cs
@@ -29,6 +29,7 @@ public class LocalizationDashboard : ModuleRules
                 "MainFrame",
                 "SourceControl",
                 "SharedSettingsWidgets",
+                "LocalizationService"
 			}
 		);
 
diff --git a/Engine/Source/Editor/LocalizationDashboard/Private/LocalizationDashboardModule.cpp b/Engine/Source/Editor/LocalizationDashboard/Private/LocalizationDashboardModule.cpp
--- a/Engine/Source/Editor/LocalizationDashboard/Private/LocalizationDashboardModule.cpp
+++ b/Engine/Source/Editor/LocalizationDashboard/Private/LocalizationDashboardModule.cpp
@@ -10,6 +10,7 @@
 #include "LocalizationTargetSetDetailCustomization.h"
 #include "LocalizationTargetDetailCustomization.h"
 #include "GatherTextDetailCustomizations.h"
+#include "PropertyEditorModule.h"
 
 #define LOCTEXT_NAMESPACE "LocalizationDashboard"
 
@@ -21,7 +22,6 @@ public:
 	virtual void StartupModule() override
 	{
 		ServiceProviders = IModularFeatures::Get().GetModularFeatureImplementations<ILocalizationServiceProvider>("LocalizationService");
-		ServiceProviders.Insert(nullptr, 0); // "None"
 
 		FPropertyEditorModule& PropertyModule = FModuleManager::LoadModuleChecked<FPropertyEditorModule>("PropertyEditor");
 		PropertyModule.RegisterCustomClassLayout("LocalizationTargetSet", FOnGetDetailCustomizationInstance::CreateLambda(
diff --git a/Engine/Source/Editor/LocalizationDashboard/Private/LocalizationTargetDetailCustomization.cpp b/Engine/Source/Editor/LocalizationDashboard/Private/LocalizationTargetDetailCustomization.cpp
--- a/Engine/Source/Editor/LocalizationDashboard/Private/LocalizationTargetDetailCustomization.cpp
+++ b/Engine/Source/Editor/LocalizationDashboard/Private/LocalizationTargetDetailCustomization.cpp
@@ -15,6 +15,8 @@
 #include "ObjectEditorUtils.h"
 #include "LocalizationDashboardSettings.h"
 #include "SErrorText.h"
+#include "ILocalizationServiceModule.h"
+#include "ILocalizationServiceProvider.h"
 
 #define LOCTEXT_NAMESPACE "LocalizationTargetEditor"
 
@@ -283,6 +285,11 @@ void FLocalizationTargetDetailCustomization::CustomizeDetails(IDetailLayoutBuild
 			];
 	});
 
+	{
+		// The sort priority is set the first time we edit the category, so set it here first
+		IDetailCategoryBuilder& DetailCategoryBuilder = DetailBuilder.EditCategory("Target", LOCTEXT("TargetCategoryLabel","Target"), ECategoryPriority::Variable);
+	}
+
 	UStructProperty* const SettingsStructProperty = CastChecked<UStructProperty>(TargetSettingsPropertyHandle->GetProperty());
 	for (TFieldIterator<UProperty> Iterator(SettingsStructProperty->Struct); Iterator; ++Iterator)
 	{
@@ -352,16 +359,16 @@ void FLocalizationTargetDetailCustomization::CustomizeDetails(IDetailLayoutBuild
 			];
 	}
 
+	{
+		// Assign this to the "important" category to make it show up just under the "target" category
+		IDetailCategoryBuilder& ServiceProviderCategoryBuilder = DetailBuilder.EditCategory(FName("LocalizationServiceProvider"), LOCTEXT("LocalizationServiceProviderCategoryLabel","Localization Service Provider"), ECategoryPriority::Important);
 
-	//{
-	//	IDetailCategoryBuilder& ServiceProviderCategoryBuilder = DetailBuilder.EditCategory(FName("LocalizationServiceProvider"));
-
-	//	ILocalizationServiceProvider* const LSP = ILocalizationDashboardModule::Get().GetCurrentLocalizationServiceProvider();
-	//	if (LSP && LocalizationTarget)
-	//	{
-	//		LSP->CustomizeTargetDetails(ServiceProviderCategoryBuilder, *LocalizationTarget);
-	//	}
-	//}
+		const ILocalizationServiceProvider& LSP = ILocalizationServiceModule::Get().GetProvider();
+		if (LocalizationTarget.IsValid())
+		{
+			LSP.CustomizeTargetDetails(ServiceProviderCategoryBuilder, LocalizationTarget->Settings.Guid);
+		}
+	}
 }
 
 FLocalizationTargetSettings* FLocalizationTargetDetailCustomization::GetTargetSettings() const
diff --git a/Engine/Source/Editor/LocalizationDashboard/Private/LocalizationTargetSetDetailCustomization.cpp b/Engine/Source/Editor/LocalizationDashboard/Private/LocalizationTargetSetDetailCustomization.cpp
--- a/Engine/Source/Editor/LocalizationDashboard/Private/LocalizationTargetSetDetailCustomization.cpp
+++ b/Engine/Source/Editor/LocalizationDashboard/Private/LocalizationTargetSetDetailCustomization.cpp
@@ -10,38 +10,41 @@
 #include "IDetailsView.h"
 #include "LocalizationCommandletTasks.h"
 #include "ObjectEditorUtils.h"
+#include "ILocalizationServiceProvider.h"
+#include "ILocalizationServiceModule.h"
+#include "ILocalizationDashboardModule.h"
 
 #define LOCTEXT_NAMESPACE "LocalizationDashboard"
 
 FLocalizationTargetSetDetailCustomization::FLocalizationTargetSetDetailCustomization()
 	: DetailLayoutBuilder(nullptr)
-	//, ServiceProviderCategoryBuilder(nullptr)
+	, ServiceProviderCategoryBuilder(nullptr)
 	, NewEntryIndexToBeInitialized(INDEX_NONE)
 {
-	//TArray<ILocalizationServiceProvider*> ActualProviders = ILocalizationDashboardModule::Get().GetLocalizationServiceProviders();
-	//for (ILocalizationServiceProvider* ActualProvider : ActualProviders)
-	//{
-	//	TSharedPtr<FLocalizationServiceProviderWrapper> Provider = MakeShareable(new FLocalizationServiceProviderWrapper(ActualProvider));
-	//	Providers.Add(Provider);
-	//}
+	TArray<ILocalizationServiceProvider*> ActualProviders = ILocalizationDashboardModule::Get().GetLocalizationServiceProviders();
+	for (ILocalizationServiceProvider* ActualProvider : ActualProviders)
+	{
+		TSharedPtr<FLocalizationServiceProviderWrapper> Provider = MakeShareable(new FLocalizationServiceProviderWrapper(ActualProvider));
+		Providers.Add(Provider);
+	}
 }
 
-//class FindProviderPredicate
-//{
-//public:
-//	FindProviderPredicate(ILocalizationServiceProvider* const InActualProvider)
-//		: ActualProvider(InActualProvider)
-//	{
-//	}
+class FindProviderPredicate
+{
+public:
+	FindProviderPredicate(ILocalizationServiceProvider* const InActualProvider)
+		: ActualProvider(InActualProvider)
+	{
+	}
 
-//	bool operator()(const TSharedPtr<FLocalizationServiceProviderWrapper>& Provider)
-//	{
-//		return Provider->Provider == ActualProvider;
-//	}
+	bool operator()(const TSharedPtr<FLocalizationServiceProviderWrapper>& Provider)
+	{
+		return Provider->Provider == ActualProvider;
+	}
 
-//private:
-//	ILocalizationServiceProvider* ActualProvider;
-//};
+private:
+	ILocalizationServiceProvider* ActualProvider;
+};
 
 namespace
 {
@@ -175,6 +178,49 @@ void FLocalizationTargetSetDetailCustomization::CustomizeDetails(IDetailLayoutBu
 				];
 		}
 	}
+
+	ServiceProviderCategoryBuilder = &(DetailLayoutBuilder->EditCategory("ServiceProvider", LOCTEXT("LocalizationServiceProvider", "Localization Service Provider"), ECategoryPriority::Important));
+	FDetailWidgetRow& DetailWidgetRow = ServiceProviderCategoryBuilder->AddCustomRow(LOCTEXT("Selected LocalizationServiceProvider", "Selected Localization Service Provider"));
+
+	int32 CurrentlySelectedProviderIndex = 0;
+
+	for (int ProviderIndex = 0; ProviderIndex < Providers.Num(); ++ProviderIndex)
+	{
+		FName CurrentlySelectedProviderName = ILocalizationServiceModule::Get().GetProvider().GetName();
+		if (Providers[ProviderIndex].IsValid() && Providers[ProviderIndex]->Provider && Providers[ProviderIndex]->Provider->GetName() == CurrentlySelectedProviderName)
+		{
+			CurrentlySelectedProviderIndex = ProviderIndex;
+			break;
+		}
+	}
+
+	DetailWidgetRow.NameContent()
+		[
+			SNew(STextBlock)
+			.Text(LOCTEXT("LocalizationServiceProvider", "Localization Service Provider"))
+		];
+	DetailWidgetRow.ValueContent()
+		[
+			SNew(SComboBox< TSharedPtr<FLocalizationServiceProviderWrapper>>)
+			.OptionsSource(&(Providers))
+			.OnSelectionChanged(this, &FLocalizationTargetSetDetailCustomization::ServiceProviderComboBox_OnSelectionChanged)
+			.OnGenerateWidget(this, &FLocalizationTargetSetDetailCustomization::ServiceProviderComboBox_OnGenerateWidget)
+			.InitiallySelectedItem(Providers[CurrentlySelectedProviderIndex])
+			.Content()
+			[
+				SNew(STextBlock)
+				.Text_Lambda([]()
+				{
+					return ILocalizationServiceModule::Get().GetProvider().GetDisplayName();
+				})
+			]
+		];
+
+	const ILocalizationServiceProvider& LSP = ILocalizationServiceModule::Get().GetProvider();
+	if (ServiceProviderCategoryBuilder != nullptr)
+	{
+		LSP.CustomizeSettingsDetails(*ServiceProviderCategoryBuilder);
+	}
 }
 
 void FLocalizationTargetSetDetailCustomization::BuildTargetsList()
@@ -237,31 +283,33 @@ void FLocalizationTargetSetDetailCustomization::RebuildTargetsList()
 	}
 }
 
-//FText FLocalizationTargetSetDetailCustomization::GetCurrentServiceProviderDisplayName() const
-//{
-//	ILocalizationServiceProvider* const LSP = ILocalizationDashboardModule::Get().GetCurrentLocalizationServiceProvider();
-//	return LSP ? LSP->GetDisplayName() : LOCTEXT("NoServiceProviderName", "None");
-//}
-//
-//TSharedRef<SWidget> FLocalizationTargetSetDetailCustomization::ServiceProviderComboBox_OnGenerateWidget(TSharedPtr<FLocalizationServiceProviderWrapper> LSPWrapper) const
-//{
-//	ILocalizationServiceProvider* const LSP = LSPWrapper->Provider;
-//	return	SNew(STextBlock)
-//		.Text(LSP ? LSP->GetDisplayName() : LOCTEXT("NoServiceProviderName", "None"));
-//}
-//
-//void FLocalizationTargetSetDetailCustomization::ServiceProviderComboBox_OnSelectionChanged(TSharedPtr<FLocalizationServiceProviderWrapper> LSPWrapper, ESelectInfo::Type SelectInfo)
-//{
-//	ILocalizationServiceProvider* const LSP = LSPWrapper->Provider;
-//	FString ServiceProviderName = LSP ? LSP->GetName().ToString() : TEXT("None");
-//	ServiceProviderPropertyHandle->SetValue(ServiceProviderName);
-//
-//	if (LSP)
-//	{
-//		LSP->CustomizeSettingsDetails(*ServiceProviderCategoryBuilder);
-//	}
-//	DetailLayoutBuilder->ForceRefreshDetails();
-//}
+FText FLocalizationTargetSetDetailCustomization::GetCurrentServiceProviderDisplayName() const
+{
+	const ILocalizationServiceProvider& LSP = ILocalizationServiceModule::Get().GetProvider();
+	return LSP.GetDisplayName();
+}
+
+TSharedRef<SWidget> FLocalizationTargetSetDetailCustomization::ServiceProviderComboBox_OnGenerateWidget(TSharedPtr<FLocalizationServiceProviderWrapper> LSPWrapper) const
+{
+	ILocalizationServiceProvider* LSP = LSPWrapper.IsValid() ? LSPWrapper->Provider : nullptr;
+
+	return	SNew(STextBlock)
+		.Text(LSP ? LSP->GetDisplayName() : LOCTEXT("NoServiceProviderName", "None"));
+}
+
+void FLocalizationTargetSetDetailCustomization::ServiceProviderComboBox_OnSelectionChanged(TSharedPtr<FLocalizationServiceProviderWrapper> LSPWrapper, ESelectInfo::Type SelectInfo)
+{
+	ILocalizationServiceProvider* LSP = LSPWrapper.IsValid() ? LSPWrapper->Provider : nullptr;
+
+	FName ServiceProviderName = LSP ? LSP->GetName() : FName(TEXT("None"));
+	ILocalizationServiceModule::Get().SetProvider(ServiceProviderName);
+
+	if (LSP && ServiceProviderCategoryBuilder)
+	{
+		LSP->CustomizeSettingsDetails(*ServiceProviderCategoryBuilder);
+	}
+	DetailLayoutBuilder->ForceRefreshDetails();
+}
 
 bool FLocalizationTargetSetDetailCustomization::CanGatherAllTargets() const
 {
diff --git a/Engine/Source/Editor/LocalizationDashboard/Private/LocalizationTargetSetDetailCustomization.h b/Engine/Source/Editor/LocalizationDashboard/Private/LocalizationTargetSetDetailCustomization.h
--- a/Engine/Source/Editor/LocalizationDashboard/Private/LocalizationTargetSetDetailCustomization.h
+++ b/Engine/Source/Editor/LocalizationDashboard/Private/LocalizationTargetSetDetailCustomization.h
@@ -3,20 +3,20 @@
 #pragma once
 
 #include "IDetailCustomization.h"
+#include "ILocalizationServiceProvider.h"
 
 class ULocalizationTarget;
 class IPropertyHandle;
 class ULocalizationTargetSet;
+class IDetailCategoryBuilder;
 
-//class ILocalizationServiceProvider;
+struct FLocalizationServiceProviderWrapper
+{
+	FLocalizationServiceProviderWrapper() : Provider(nullptr) {}
+	FLocalizationServiceProviderWrapper(ILocalizationServiceProvider* const InProvider) : Provider(InProvider) {}
 
-//struct FLocalizationServiceProviderWrapper
-//{
-//	FLocalizationServiceProviderWrapper() : Provider(nullptr) {}
-//	FLocalizationServiceProviderWrapper(ILocalizationServiceProvider* const InProvider) : Provider(InProvider) {}
-//
-//	ILocalizationServiceProvider* Provider;
-//};
+	ILocalizationServiceProvider* Provider;
+};
 
 class FLocalizationTargetSetDetailCustomization : public IDetailCustomization
 {
@@ -28,9 +28,9 @@ private:
 	void BuildTargetsList();
 	void RebuildTargetsList();
 
-	//FText GetCurrentServiceProviderDisplayName() const;
-	//TSharedRef<SWidget> ServiceProviderComboBox_OnGenerateWidget(TSharedPtr<FLocalizationServiceProviderWrapper> LSPWrapper) const;
-	//void ServiceProviderComboBox_OnSelectionChanged(TSharedPtr<FLocalizationServiceProviderWrapper> LSPWrapper, ESelectInfo::Type SelectInfo);
+	FText GetCurrentServiceProviderDisplayName() const;
+	TSharedRef<SWidget> ServiceProviderComboBox_OnGenerateWidget(TSharedPtr<FLocalizationServiceProviderWrapper> LSPWrapper) const;
+	void ServiceProviderComboBox_OnSelectionChanged(TSharedPtr<FLocalizationServiceProviderWrapper> LSPWrapper, ESelectInfo::Type SelectInfo);
 
 	bool CanGatherAllTargets() const;
 	void GatherAllTargets();
@@ -52,9 +52,8 @@ private:
 
 	TWeakObjectPtr<ULocalizationTargetSet> TargetSet;
 
-	//TSharedPtr<IPropertyHandle> ServiceProviderPropertyHandle;
-	//IDetailCategoryBuilder* ServiceProviderCategoryBuilder;
-	//TArray< TSharedPtr<FLocalizationServiceProviderWrapper> > Providers;
+	IDetailCategoryBuilder* ServiceProviderCategoryBuilder;
+	TArray< TSharedPtr<FLocalizationServiceProviderWrapper> > Providers;
 
 	TSharedPtr<IPropertyHandle> TargetObjectsPropertyHandle;
 	FSimpleDelegate TargetsArrayPropertyHandle_OnNumElementsChanged;
diff --git a/Engine/Source/Editor/LocalizationDashboard/Public/ILocalizationServiceProvider.h b/Engine/Source/Editor/LocalizationDashboard/Public/ILocalizationServiceProvider.h
deleted file mode 100644
