Fixing the build.

diff --git a/Engine/Source/Editor/LocalizationDashboard/Private/LocalizationTargetTypes.h b/Engine/Source/Editor/LocalizationDashboard/Private/LocalizationTargetTypes.h
--- a/Engine/Source/Editor/LocalizationDashboard/Private/LocalizationTargetTypes.h
+++ b/Engine/Source/Editor/LocalizationDashboard/Private/LocalizationTargetTypes.h
@@ -241,7 +241,7 @@ enum class ELocalizationTargetLoadingPolicy : uint8
 	/** This target's localization data will only be loaded if the editor is displaying localized property names. */
 	PropertyNames,
 	/** This target's localization data will only be loaded if the editor is displaying localized tool tips. */
-	Tooltips,
+	ToolTips,
 };
 
 USTRUCT()
