Localization Dashboard: Fix crash on loading Loc Dashboard if OneSky plugin was not enabled.

diff --git a/Engine/Source/Editor/LocalizationDashboard/Private/LocalizationDashboardModule.cpp b/Engine/Source/Editor/LocalizationDashboard/Private/LocalizationDashboardModule.cpp
--- a/Engine/Source/Editor/LocalizationDashboard/Private/LocalizationDashboardModule.cpp
+++ b/Engine/Source/Editor/LocalizationDashboard/Private/LocalizationDashboardModule.cpp
@@ -11,6 +11,7 @@
 #include "LocalizationTargetDetailCustomization.h"
 #include "GatherTextDetailCustomizations.h"
 #include "PropertyEditorModule.h"
+#include "ILocalizationServiceModule.h"
 
 #define LOCTEXT_NAMESPACE "LocalizationDashboard"
 
@@ -21,6 +22,8 @@ public:
 	// Begin IModuleInterface interface
 	virtual void StartupModule() override
 	{
+		// Ensure Localization Service module is loaded (registers default localization service on startup)
+		FModuleManager::LoadModuleChecked<ILocalizationServiceModule>("LocalizationService");
 		ServiceProviders = IModularFeatures::Get().GetModularFeatureImplementations<ILocalizationServiceProvider>("LocalizationService");
 
 		FPropertyEditorModule& PropertyModule = FModuleManager::LoadModuleChecked<FPropertyEditorModule>("PropertyEditor");
