Fixed localization configuration script generation code that wasn't handling relative paths to files on different drives appropriately. If a relative path can not be calculated, an absolute path will be used.

diff --git a/Engine/Source/Editor/LocalizationDashboard/Private/LocalizationConfigurationScript.cpp b/Engine/Source/Editor/LocalizationDashboard/Private/LocalizationConfigurationScript.cpp
--- a/Engine/Source/Editor/LocalizationDashboard/Private/LocalizationConfigurationScript.cpp
+++ b/Engine/Source/Editor/LocalizationDashboard/Private/LocalizationConfigurationScript.cpp
@@ -24,7 +24,10 @@ namespace LocalizationConfigurationScript
 	{
 		FString Result = Path;
 		const FString ProjectDir = !IsUsingProjectFile ? FPaths::EngineDir() : FPaths::GameDir();
-		FPaths::MakePathRelativeTo(Result, *ProjectDir);
+		if (!FPaths::MakePathRelativeTo(Result, *ProjectDir))
+		{
+			Result = FPaths::ConvertRelativePathToFull(Path);
+		}
 		return Result;
 	}
 
