Move LocalizationTargetTypes.h to Public folder so it can be included by other modules.

diff --git a/Engine/Source/Editor/LocalizationDashboard/Private/LocalizationTargetTypes.h b/Engine/Source/Editor/LocalizationDashboard/Private/LocalizationTargetTypes.h
deleted file mode 100644
diff --git a/Engine/Source/Editor/LocalizationDashboard/Public/LocalizationTargetTypes.h b/Engine/Source/Editor/LocalizationDashboard/Public/LocalizationTargetTypes.h
new file mode 100644
--- /dev/null
+++ b/Engine/Source/Editor/LocalizationDashboard/Public/LocalizationTargetTypes.h
@@ -0,0 +1,337 @@
+﻿// Copyright 1998-2015 Epic Games, Inc. All Rights Reserved.
+
+#pragma once
+
+#include "Engine/EngineTypes.h"
+#include "LocalizationTargetTypes.generated.h"
+
+UENUM()
+enum class ELocalizationTargetConflictStatus : uint8
+{
+	/** The status of conflicts in this localization target could not be determined. */
+	Unknown,
+	/** The are outstanding conflicts present in this localization target. */
+	ConflictsPresent,
+	/** The localization target is clear of conflicts. */
+	Clear,
+};
+
+USTRUCT()
+struct FGatherTextSearchDirectory
+{
+	GENERATED_USTRUCT_BODY()
+
+	UPROPERTY(config, EditAnywhere, Category="Path")
+	FString Path;
+
+	bool Validate(const FString& RootDirectory, FText& OutError) const;
+};
+
+USTRUCT()
+struct FGatherTextIncludePath
+{
+	GENERATED_USTRUCT_BODY()
+
+	UPROPERTY(config, EditAnywhere, Category="Pattern")
+	FString Pattern;
+
+	bool Validate(const FString& RootDirectory, FText& OutError) const;
+};
+
+USTRUCT()
+struct FGatherTextExcludePath
+{
+	GENERATED_USTRUCT_BODY()
+
+	UPROPERTY(config, EditAnywhere, Category="Pattern")
+	FString Pattern;
+
+	bool Validate(FText& OutError) const;
+};
+
+USTRUCT()
+struct FGatherTextFileExtension
+{
+	GENERATED_USTRUCT_BODY()
+
+	UPROPERTY(config, EditAnywhere, Category="Pattern")
+	FString Pattern;
+
+	bool Validate(FText& OutError) const;
+};
+
+USTRUCT()
+struct FGatherTextFromTextFilesConfiguration
+{
+	GENERATED_USTRUCT_BODY()
+
+	static const TArray<FGatherTextFileExtension> DefaultTextFileExtensions;
+
+	FGatherTextFromTextFilesConfiguration()
+		: IsEnabled(true)
+		, FileExtensions(DefaultTextFileExtensions)
+	{
+	}
+
+	/* If enabled, text from text files will be gathered according to this configuration. */
+	UPROPERTY(config, EditAnywhere, Category = "Execution")
+	bool IsEnabled;
+
+	/* The paths of directories to be searched recursively for text files, specified relative to the project's root, which may be parsed for text to gather. */
+	UPROPERTY(config, EditAnywhere, Category = "Filter")
+	TArray<FGatherTextSearchDirectory> SearchDirectories;
+
+	/* Text files whose paths match these wildcard patterns will be excluded from gathering. */
+	UPROPERTY(config, EditAnywhere, Category = "Filter")
+	TArray<FGatherTextExcludePath> ExcludePathWildcards;
+
+	/* Text files whose names match these wildcard patterns may be parsed for text to gather. */
+	UPROPERTY(config, EditAnywhere, Category = "Filter")
+	TArray<FGatherTextFileExtension> FileExtensions;
+
+	bool Validate(const FString& RootDirectory, FText& OutError) const;
+};
+
+
+USTRUCT()
+struct FGatherTextFromPackagesConfiguration
+{
+	GENERATED_USTRUCT_BODY()
+
+	static const TArray<FGatherTextFileExtension> DefaultPackageFileExtensions;
+
+	FGatherTextFromPackagesConfiguration()
+		: IsEnabled(true)
+		, FileExtensions(DefaultPackageFileExtensions)
+	{
+	}
+
+	/* If enabled, text from packages will be gathered according to this configuration. */
+	UPROPERTY(config, EditAnywhere, Category = "Execution")
+	bool IsEnabled;
+
+	/* Packages whose paths match these wildcard patterns, specified relative to the project's root, may be processed for gathering. */
+	UPROPERTY(config, EditAnywhere, Category = "Filter")
+	TArray<FGatherTextIncludePath> IncludePathWildcards;
+
+	/* Packages whose paths match these wildcard patterns will be excluded from gathering. */
+	UPROPERTY(config, EditAnywhere, Category = "Filter")
+	TArray<FGatherTextExcludePath> ExcludePathWildcards;
+
+	/* Packages whose names match these wildcard patterns may be processed for text to gather. */
+	UPROPERTY(config, EditAnywhere, Category = "Filter")
+	TArray<FGatherTextFileExtension> FileExtensions;
+
+	bool Validate(const FString& RootDirectory, FText& OutError) const;
+};
+
+USTRUCT()
+struct FMetaDataTextKeyPattern
+{
+	GENERATED_USTRUCT_BODY()
+
+	UPROPERTY(config, EditAnywhere, Category="Pattern")
+	FString Pattern;
+
+	static const TArray<FString> PossiblePlaceHolders;
+
+	bool Validate(FText& OutError) const;
+};
+
+USTRUCT()
+struct FMetaDataKeyName
+{
+	GENERATED_USTRUCT_BODY()
+
+	UPROPERTY(config, EditAnywhere, Category="Name")
+	FString Name;
+
+	bool Validate(FText& OutError) const;
+};
+
+USTRUCT()
+struct FMetaDataKeyGatherSpecification
+{
+	GENERATED_USTRUCT_BODY()
+
+	/*  The metadata key for which values will be gathered as text. */
+	UPROPERTY(config, EditAnywhere, Category = "Input")
+	FMetaDataKeyName MetaDataKey;
+
+	/* The localization namespace in which the gathered text will be output. */
+	UPROPERTY(config, EditAnywhere, Category = "Output")
+	FString TextNamespace;
+
+	/* The pattern which will be formatted to form the localization key for the metadata value gathered as text.
+	Placeholder - Description
+	{FieldPath} - The fully qualified name of the object upon which the metadata resides.
+	{MetaDataValue} - The value associated with the metadata key. */
+	UPROPERTY(config, EditAnywhere, Category = "Output")
+	FMetaDataTextKeyPattern TextKeyPattern;
+
+	bool Validate(FText& OutError) const;
+};
+
+USTRUCT()
+struct FGatherTextFromMetaDataConfiguration
+{
+	GENERATED_USTRUCT_BODY()
+
+	FGatherTextFromMetaDataConfiguration()
+		: IsEnabled(false)
+	{
+	}
+
+	/* If enabled, metadata will be gathered according to this configuration. */
+	UPROPERTY(config, EditAnywhere, Category = "Execution")
+	bool IsEnabled;
+
+	/* Metadata from source files whose paths match these wildcard patterns, specified relative to the project's root, may be processed for gathering. */
+	UPROPERTY(config, EditAnywhere, Category = "Filter")
+	TArray<FGatherTextIncludePath> IncludePathWildcards;
+
+	/* Metadata from source files whose paths match these wildcard patterns will be excluded from gathering. */
+	UPROPERTY(config, EditAnywhere, Category = "Filter")
+	TArray<FGatherTextExcludePath> ExcludePathWildcards;
+
+	/* Specifications for how to gather text from specific metadata keys. */
+	UPROPERTY(config, EditAnywhere, Category = "MetaData")
+	TArray<FMetaDataKeyGatherSpecification> KeySpecifications;
+
+	bool Validate(const FString& RootDirectory, FText& OutError) const;
+};
+
+USTRUCT()
+struct FCultureStatistics
+{
+	GENERATED_USTRUCT_BODY()
+
+	FCultureStatistics()
+		: WordCount(0)
+	{
+	}
+
+	FCultureStatistics(const FString& InCultureName)
+		: CultureName(InCultureName)
+		, WordCount(0)
+	{
+	}
+
+
+	/* The ISO name for this culture. */
+	UPROPERTY(config, EditAnywhere, Category = "Culture")
+	FString CultureName;
+
+	/* The estimated number of words that have been localized for this culture. */
+	UPROPERTY(Transient, EditAnywhere, Category = "Statistics")
+	uint32 WordCount;
+};
+
+UENUM()
+enum class ELocalizationTargetLoadingPolicy : uint8
+{
+	/** This target's localization data will never be loaded automatically. */
+	Never,
+	/** This target's localization data will always be loaded automatically. */
+	Always,
+	/** This target's localization data will only be loaded when running the editor. Use if this target localizes the editor. */
+	Editor,
+	/** This target's localization data will only be loaded when running the game. Use if this target localizes your game. */
+	Game,
+	/** This target's localization data will only be loaded if the editor is displaying localized property names. */
+	PropertyNames,
+	/** This target's localization data will only be loaded if the editor is displaying localized tool tips. */
+	ToolTips,
+};
+
+USTRUCT()
+struct FLocalizationTargetSettings
+{
+	GENERATED_USTRUCT_BODY()
+
+	FLocalizationTargetSettings()
+		: Guid(FGuid::NewGuid())
+		, ConflictStatus(ELocalizationTargetConflictStatus::Unknown)
+		, NativeCultureIndex(INDEX_NONE)
+	{
+	}
+
+	/* Unique name for the target. */
+	UPROPERTY(config, EditAnywhere, Category = "Target")
+	FString Name;
+
+	UPROPERTY(config)
+	FGuid Guid;
+
+	/* Whether the target has outstanding conflicts that require resolution. */
+	UPROPERTY(Transient, EditAnywhere, Category = "Target")
+	ELocalizationTargetConflictStatus ConflictStatus;
+
+	/* Text present in these targets will not be duplicated in this target. */
+	UPROPERTY(config, EditAnywhere, Category = "Gather")
+	TArray<FGuid> TargetDependencies;
+
+	/* Text present in these manifests will not be duplicated in this target. */
+	UPROPERTY(config, EditAnywhere, Category = "Gather", AdvancedDisplay, meta=(FilePathFilter="manifest"))
+	TArray<FFilePath> AdditionalManifestDependencies;
+
+	/* The names of modules which must be loaded when gathering text. Use to gather from packages or metadata sourced from a module that isn't the primary game module. */
+	UPROPERTY(config, EditAnywhere, Category = "Gather", AdvancedDisplay)
+	TArray<FString> RequiredModuleNames;
+
+	/* Parameters for defining what text is gathered from text files and how. */
+	UPROPERTY(config, EditAnywhere, Category = "Gather")
+	FGatherTextFromTextFilesConfiguration GatherFromTextFiles;
+
+	/* Parameters for defining what text is gathered from packages and how. */
+	UPROPERTY(config, EditAnywhere, Category = "Gather")
+	FGatherTextFromPackagesConfiguration GatherFromPackages;
+
+	/* Parameters for defining what text is gathered from metadata and how. */
+	UPROPERTY(config, EditAnywhere, Category = "Gather")
+	FGatherTextFromMetaDataConfiguration GatherFromMetaData;
+
+	/* The index of the native culture among the supported cultures. */
+	UPROPERTY(config, EditAnywhere, Category = "Cultures")
+	int32 NativeCultureIndex;
+
+	/* Cultures for which the source text is being localized for.*/
+	UPROPERTY(config, EditAnywhere, Category = "Cultures")
+	TArray<FCultureStatistics> SupportedCulturesStatistics;
+};
+
+UCLASS(Within=LocalizationTargetSet)
+class ULocalizationTarget : public UObject
+{
+	GENERATED_BODY()
+
+public:
+	UPROPERTY(EditAnywhere, Category = "Target")
+	FLocalizationTargetSettings Settings;
+
+public:
+#if WITH_EDITOR
+	virtual void PostEditChangeProperty(FPropertyChangedEvent& PropertyChangedEvent) override;
+#endif
+
+	bool IsMemberOfEngineTargetSet() const;
+	bool UpdateWordCountsFromCSV();
+	void UpdateStatusFromConflictReport();
+	bool RenameTargetAndFiles(const FString& NewName);
+	bool DeleteFiles(const FString* const Culture = nullptr) const;
+};
+
+UCLASS(Within=LocalizationDashboardSettings)
+class ULocalizationTargetSet : public UObject
+{
+	GENERATED_BODY()
+
+public:
+	UPROPERTY(EditAnywhere, Category = "Targets")
+	TArray<ULocalizationTarget*> TargetObjects;
+
+public:
+#if WITH_EDITOR
+	virtual void PostEditChangeProperty(FPropertyChangedEvent& PropertyChangedEvent) override;
+#endif
+};
\ No newline at end of file
