Localization dashboard now appropriately adds the native culture setting to generated localization configuration scripts.

diff --git a/Engine/Source/Editor/LocalizationDashboard/Private/LocalizationConfigurationScript.cpp b/Engine/Source/Editor/LocalizationDashboard/Private/LocalizationConfigurationScript.cpp
--- a/Engine/Source/Editor/LocalizationDashboard/Private/LocalizationConfigurationScript.cpp
+++ b/Engine/Source/Editor/LocalizationDashboard/Private/LocalizationConfigurationScript.cpp
@@ -135,7 +135,7 @@ namespace LocalizationConfigurationScript
 
 			if (Target->Settings.SupportedCulturesStatistics.IsValidIndex(Target->Settings.NativeCultureIndex))
 			{
-			ConfigSection.Add( TEXT("NativeCulture"), Target->Settings.SupportedCulturesStatistics[Target->Settings.NativeCultureIndex].CultureName );
+				ConfigSection.Add( TEXT("NativeCulture"), Target->Settings.SupportedCulturesStatistics[Target->Settings.NativeCultureIndex].CultureName );
 			}
 			for (const FCultureStatistics& CultureStatistics : Target->Settings.SupportedCulturesStatistics)
 			{
@@ -309,6 +309,11 @@ namespace LocalizationConfigurationScript
 			const FString DestinationPath = ContentDirRelativeToGameDir / TEXT("Localization") / Target->Settings.Name;
 			ConfigSection.Add( TEXT("DestinationPath"), DestinationPath );
 
+			if (Target->Settings.SupportedCulturesStatistics.IsValidIndex(Target->Settings.NativeCultureIndex))
+			{
+				ConfigSection.Add( TEXT("NativeCulture"), Target->Settings.SupportedCulturesStatistics[Target->Settings.NativeCultureIndex].CultureName );
+			}
+
 			const auto& AddCultureToGenerate = [&](const int32 Index)
 			{
 				ConfigSection.Add( TEXT("CulturesToGenerate"), Target->Settings.SupportedCulturesStatistics[Index].CultureName );
@@ -411,6 +416,11 @@ namespace LocalizationConfigurationScript
 			}
 			ConfigSection.Add( TEXT("DestinationPath"), DestinationPath );
 
+			if (Target->Settings.SupportedCulturesStatistics.IsValidIndex(Target->Settings.NativeCultureIndex))
+			{
+				ConfigSection.Add( TEXT("NativeCulture"), Target->Settings.SupportedCulturesStatistics[Target->Settings.NativeCultureIndex].CultureName );
+			}
+
 			const auto& AddCultureToGenerate = [&](const int32 Index)
 			{
 				ConfigSection.Add( TEXT("CulturesToGenerate"), Target->Settings.SupportedCulturesStatistics[Index].CultureName );
@@ -535,6 +545,11 @@ namespace LocalizationConfigurationScript
 			ConfigSection.Add( TEXT("ManifestName"), FString::Printf( TEXT("%s.%s"), *Target->Settings.Name, TEXT("manifest") ) );
 			ConfigSection.Add( TEXT("ResourceName"), FString::Printf( TEXT("%s.%s"), *Target->Settings.Name, TEXT("locres") ) );
 
+			if (Target->Settings.SupportedCulturesStatistics.IsValidIndex(Target->Settings.NativeCultureIndex))
+			{
+				ConfigSection.Add( TEXT("NativeCulture"), Target->Settings.SupportedCulturesStatistics[Target->Settings.NativeCultureIndex].CultureName );
+			}
+
 			const auto& AddCultureToGenerate = [&](const int32 Index)
 			{
 				ConfigSection.Add( TEXT("CulturesToGenerate"), Target->Settings.SupportedCulturesStatistics[Index].CultureName );
