Translation Editor/Picker: Generate the "RegenerateTarget.ini" config files on demand when Previewing translations in the editor

diff --git a/Engine/Config/Localization/RegenerateEditor.ini b/Engine/Config/Localization/RegenerateEditor.ini
deleted file mode 100644
diff --git a/Engine/Config/Localization/RegenerateEditorTutorials.ini b/Engine/Config/Localization/RegenerateEditorTutorials.ini
deleted file mode 100644
diff --git a/Engine/Config/Localization/RegenerateEngine.ini b/Engine/Config/Localization/RegenerateEngine.ini
deleted file mode 100644
diff --git a/Engine/Source/Editor/LocalizationDashboard/Private/LocalizationConfigurationScript.cpp b/Engine/Source/Editor/LocalizationDashboard/Private/LocalizationConfigurationScript.cpp
--- a/Engine/Source/Editor/LocalizationDashboard/Private/LocalizationConfigurationScript.cpp
+++ b/Engine/Source/Editor/LocalizationDashboard/Private/LocalizationConfigurationScript.cpp
@@ -585,9 +585,46 @@ namespace LocalizationConfigurationScript
 			ConfigFilePath = ConfigFileDirectory / FString::Printf( TEXT("%s_Compile_%s.%s"), *Target->Settings.Name, *CultureName.GetValue(), TEXT("ini") );
 		}
 		else
-	{
+		{
 			ConfigFilePath = ConfigFileDirectory / FString::Printf( TEXT("%s_Compile.%s"), *Target->Settings.Name, TEXT("ini") );
 		}
 		return ConfigFilePath;
 	}
+
+	FLocalizationConfigurationScript GenerateRegenerateResourcesScript(const ULocalizationTarget* const Target)
+	{
+		FLocalizationConfigurationScript Script;
+
+		const FString ContentDirRelativeToGameDir = MakePathRelativeForCommandletProcess(GetContentDir(Target), !Target->IsMemberOfEngineTargetSet());
+
+		// RegenerateResources
+		{
+			FConfigSection& ConfigSection = Script.FindOrAdd("RegenerateResources");;
+
+			if (Target->Settings.SupportedCulturesStatistics.IsValidIndex(Target->Settings.NativeCultureIndex))
+			{
+				ConfigSection.Add(TEXT("NativeCulture"), Target->Settings.SupportedCulturesStatistics[Target->Settings.NativeCultureIndex].CultureName);
+			}
+
+			const FString SourcePath = ContentDirRelativeToGameDir / TEXT("Localization") / Target->Settings.Name;
+			ConfigSection.Add(TEXT("SourcePath"), SourcePath);
+			const FString DestinationPath = ContentDirRelativeToGameDir / TEXT("Localization") / Target->Settings.Name;
+			ConfigSection.Add(TEXT("DestinationPath"), DestinationPath);
+
+			ConfigSection.Add(TEXT("ManifestName"), FString::Printf(TEXT("%s.%s"), *Target->Settings.Name, TEXT("manifest")));
+
+			ConfigSection.Add(TEXT("ResourceName"), FString::Printf(TEXT("%s.%s"), *Target->Settings.Name, TEXT("locres")));
+
+		}
+
+		Script.Dirty = true;
+
+		return Script;
+	}
+
+	FString GetRegenerateResourcesScriptPath(const ULocalizationTarget* const Target)
+	{
+		return GetScriptDirectory(Target) / FString::Printf(TEXT("Regenerate%s.%s"), *(Target->Settings.Name), TEXT("ini"));
+	}
+
 }
\ No newline at end of file
diff --git a/Engine/Source/Editor/LocalizationDashboard/Private/LocalizationConfigurationScript.h b/Engine/Source/Editor/LocalizationDashboard/Private/LocalizationConfigurationScript.h
deleted file mode 100644
diff --git a/Engine/Source/Editor/LocalizationDashboard/Public/LocalizationConfigurationScript.h b/Engine/Source/Editor/LocalizationDashboard/Public/LocalizationConfigurationScript.h
new file mode 100644
--- /dev/null
+++ b/Engine/Source/Editor/LocalizationDashboard/Public/LocalizationConfigurationScript.h
@@ -0,0 +1,54 @@
+// Copyright 1998-2015 Epic Games, Inc. All Rights Reserved.
+
+#pragma once
+
+#include "ConfigCacheIni.h"
+#include "LocalizationTargetTypes.h"
+
+struct FLocalizationConfigurationScript : public FConfigFile
+{
+	FConfigSection& CommonSettings()
+	{
+		return FindOrAdd(TEXT("CommonSettings"));
+	}
+
+	FConfigSection& GatherTextStep(const uint32 Index)
+	{
+		return FindOrAdd( FString::Printf( TEXT("GatherTextStep%u"), Index) );
+	}
+};
+
+namespace LocalizationConfigurationScript
+{
+	FString MakePathRelativeForCommandletProcess(const FString& Path, const bool IsUsingProjectFile);
+
+	FString GetScriptDirectory(const ULocalizationTarget* const Target);
+	FString GetDataDirectory(const ULocalizationTarget* const Target);
+	TArray<FString> GetScriptPaths(const ULocalizationTarget* const Target);
+
+	FString GetManifestPath(const ULocalizationTarget* const Target);
+	FString GetArchivePath(const ULocalizationTarget* const Target, const FString& CultureName);
+	FString GetDefaultPOFileName(const ULocalizationTarget* const Target);
+	FString GetDefaultPOPath(const ULocalizationTarget* const Target, const FString& CultureName);
+	FString GetLocResPath(const ULocalizationTarget* const Target, const FString& CultureName);
+	FString GetWordCountCSVPath(const ULocalizationTarget* const Target);
+	FString GetConflictReportPath(const ULocalizationTarget* const Target);
+
+	FLocalizationConfigurationScript GenerateGatherScript(const ULocalizationTarget* const Target);
+	FString GetGatherScriptPath(const ULocalizationTarget* const Target);
+
+	FLocalizationConfigurationScript GenerateImportScript(const ULocalizationTarget* const Target, const TOptional<FString> CultureName = TOptional<FString>(), const TOptional<FString> OutputPathOverride = TOptional<FString>());
+	FString GetImportScriptPath(const ULocalizationTarget* const Target, const TOptional<FString> CultureName = TOptional<FString>());
+
+	FLocalizationConfigurationScript GenerateExportScript(const ULocalizationTarget* const Target, const TOptional<FString> CultureName = TOptional<FString>(), const TOptional<FString> OutputPathOverride = TOptional<FString>());
+	FString GetExportScriptPath(const ULocalizationTarget* const Target, const TOptional<FString> CultureName = TOptional<FString>());
+
+	FLocalizationConfigurationScript GenerateWordCountReportScript(const ULocalizationTarget* const Target);
+	FString GetWordCountReportScriptPath(const ULocalizationTarget* const Target);
+
+	FLocalizationConfigurationScript GenerateCompileScript(const ULocalizationTarget* const Target, const TOptional<FString> CultureName = TOptional<FString>());
+	FString GetCompileScriptPath(const ULocalizationTarget* const Target, const TOptional<FString> CultureName = TOptional<FString>());
+
+	LOCALIZATIONDASHBOARD_API FLocalizationConfigurationScript GenerateRegenerateResourcesScript(const ULocalizationTarget* const Target);
+	LOCALIZATIONDASHBOARD_API FString GetRegenerateResourcesScriptPath(const ULocalizationTarget* const Target);
+}
\ No newline at end of file
diff --git a/Engine/Source/Editor/TranslationEditor/Private/TranslationDataManager.cpp b/Engine/Source/Editor/TranslationEditor/Private/TranslationDataManager.cpp
--- a/Engine/Source/Editor/TranslationEditor/Private/TranslationDataManager.cpp
+++ b/Engine/Source/Editor/TranslationEditor/Private/TranslationDataManager.cpp
@@ -15,6 +15,7 @@
 #include "ILocalizationServiceModule.h"
 #include "Editor/LocalizationDashboard/Public/ILocalizationDashboardModule.h"
 #include "Editor/LocalizationDashboard/Public/LocalizationTargetTypes.h"
+#include "LocalizationConfigurationScript.h"
 
 
 DEFINE_LOG_CATEGORY_STATIC(LogTranslationEditor, Log, All);
@@ -569,22 +570,24 @@ void FTranslationDataManager::PreviewAllTranslationsInEditor()
 		IsEngineManifest = true;
 	}
 
-	FString ConfigDirectory;
-	if (IsEngineManifest)
+	ULocalizationTarget* Target = ILocalizationDashboardModule::Get().GetLocalizationTargetByName(FPaths::GetBaseFilename(ManifestFilePath), IsEngineManifest);
+	if (Target)
 	{
-		ConfigDirectory = FPaths::EngineConfigDir();
+		const FString ConfigFilePath = LocalizationConfigurationScript::GetRegenerateResourcesScriptPath(Target);
+		LocalizationConfigurationScript::GenerateRegenerateResourcesScript(Target).Write(ConfigFilePath);
+
+		FJsonInternationalizationArchiveSerializer LocalizationArchiveSerializer;
+		FJsonInternationalizationManifestSerializer LocalizationManifestSerializer;
+
+		FTextLocalizationManager::Get().LoadFromManifestAndArchives(ConfigFilePath, LocalizationArchiveSerializer, LocalizationManifestSerializer);
 	}
 	else
 	{
-		ConfigDirectory = FPaths::GameConfigDir();
+		FText ErrorNotify = LOCTEXT("PreviewAllTranslationsInEditorFail", "Failed to preview translations in Editor!");
+		FMessageLog TranslationEditorMessageLog("TranslationEditor");
+		TranslationEditorMessageLog.Error(ErrorNotify);
+		TranslationEditorMessageLog.Notify(ErrorNotify);
 	}
-
-	FString ConfigFilePath = ConfigDirectory / "Localization" / "Regenerate" + FPaths::GetBaseFilename(ManifestFilePath) + ".ini";
-
-	FJsonInternationalizationArchiveSerializer LocalizationArchiveSerializer;
-	FJsonInternationalizationManifestSerializer LocalizationManifestSerializer;
-
-	FTextLocalizationManager::Get().LoadFromManifestAndArchives(ConfigFilePath, LocalizationArchiveSerializer, LocalizationManifestSerializer);
 }
 
 void FTranslationDataManager::PopulateSearchResultsUsingFilter(const FString& SearchFilter)
