Allow game localization targets to gather engine sources from an out-of-source project

The localization pipeline locks the relative paths used by a target to either the engine or project root (depending on the target type), however the UE4 engine isn't localized into many languages, so project targets often have to override some of the engine localization to support their additional languages. For in-source builds, this was possible by using a project relative path to the engine directory (which was in a fixed location), however for any out-of-source projects (such as those in the "Unreal Projects" directory) there was no way to do this in a portable way.

This change adds a "path root" to every include and exclude path, which can be set to "Auto" (the default, which detects the root based on the target type), "Engine" (to force the path to be resolved relative to the engine root), or "Project" (to force the path to be resolved relative to the project root). This allows out-of-source projects to specify paths relative to the engine in a portable way.

diff --git a/Engine/Source/Developer/Localization/Private/LocalizationConfigurationScript.cpp b/Engine/Source/Developer/Localization/Private/LocalizationConfigurationScript.cpp
--- a/Engine/Source/Developer/Localization/Private/LocalizationConfigurationScript.cpp
+++ b/Engine/Source/Developer/Localization/Private/LocalizationConfigurationScript.cpp
@@ -258,8 +258,9 @@ namespace LocalizationConfigurationScript
 	{
 		FLocalizationConfigurationScript Script;
 
-		const FString ConfigDirRelativeToGameDir = MakePathRelativeForCommandletProcess(GetConfigDir(Target), !Target->IsMemberOfEngineTargetSet());
-		const FString ContentDirRelativeToGameDir = MakePathRelativeForCommandletProcess(GetContentDir(Target), !Target->IsMemberOfEngineTargetSet());
+		const bool bIsEngineTarget = Target->IsMemberOfEngineTargetSet();
+		const FString ConfigDirRelativeToGameDir = MakePathRelativeForCommandletProcess(GetConfigDir(Target), !bIsEngineTarget);
+		const FString ContentDirRelativeToGameDir = MakePathRelativeForCommandletProcess(GetContentDir(Target), !bIsEngineTarget);
 
 		// CommonSettings
 		{
@@ -274,7 +275,7 @@ namespace LocalizationConfigurationScript
 			}
 
 			// Engine targets may not depend on game targets
-			if (!Target->IsMemberOfEngineTargetSet())
+			if (!bIsEngineTarget)
 			{
 				ULocalizationTargetSet* GameTargetSet = ULocalizationSettings::GetGameTargetSet();
 				if (GameTargetSet)
@@ -289,13 +290,13 @@ namespace LocalizationConfigurationScript
 				const ULocalizationTarget* const * OtherTarget = AllLocalizationTargets.FindByPredicate([&TargetDependencyGuid](ULocalizationTarget* const InOtherTarget)->bool{return InOtherTarget->Settings.Guid == TargetDependencyGuid;});
 				if (OtherTarget && Target != *OtherTarget)
 				{
-					ConfigSection.Add( TEXT("ManifestDependencies"), MakePathRelativeForCommandletProcess(GetManifestPath(*OtherTarget), !Target->IsMemberOfEngineTargetSet()) );
+					ConfigSection.Add( TEXT("ManifestDependencies"), MakePathRelativeForCommandletProcess(GetManifestPath(*OtherTarget), !bIsEngineTarget) );
 				}
 			}
 
 			for (const FFilePath& Path : Target->Settings.AdditionalManifestDependencies)
 			{
-				ConfigSection.Add( TEXT("ManifestDependencies"), MakePathRelativeForCommandletProcess(Path.FilePath, !Target->IsMemberOfEngineTargetSet()) );
+				ConfigSection.Add( TEXT("ManifestDependencies"), MakePathRelativeForCommandletProcess(Path.FilePath, !bIsEngineTarget) );
 			}
 
 			for (const FString& ModuleName : Target->Settings.RequiredModuleNames)
@@ -333,14 +334,14 @@ namespace LocalizationConfigurationScript
 			// Include Paths
 			for (const auto& IncludePath : Target->Settings.GatherFromTextFiles.SearchDirectories)
 			{
-				ConfigSection.Add( TEXT("SearchDirectoryPaths"), IncludePath.Path );
+				ConfigSection.Add( TEXT("SearchDirectoryPaths"), FString::Printf(TEXT("%s%s"), *FLocalizationGatherPathRootUtil::GetResolvedPathRootToken(IncludePath.PathRoot), *IncludePath.Path) );
 			}
 
 			// Exclude Paths
 			ConfigSection.Add( TEXT("ExcludePathFilters"), TEXT("Config/Localization/*") );
 			for (const auto& ExcludePath : Target->Settings.GatherFromTextFiles.ExcludePathWildcards)
 			{
-				ConfigSection.Add( TEXT("ExcludePathFilters"), ExcludePath.Pattern );
+				ConfigSection.Add( TEXT("ExcludePathFilters"), FString::Printf(TEXT("%s%s"), *FLocalizationGatherPathRootUtil::GetResolvedPathRootToken(ExcludePath.PathRoot), *ExcludePath.Pattern) );
 			}
 
 			// Source File Search Filters
@@ -363,14 +364,14 @@ namespace LocalizationConfigurationScript
 			// Include Paths
 			for (const auto& IncludePath : Target->Settings.GatherFromPackages.IncludePathWildcards)
 			{
-				ConfigSection.Add( TEXT("IncludePathFilters"), IncludePath.Pattern );
+				ConfigSection.Add( TEXT("IncludePathFilters"), FString::Printf(TEXT("%s%s"), *FLocalizationGatherPathRootUtil::GetResolvedPathRootToken(IncludePath.PathRoot), *IncludePath.Pattern) );
 			}
 
 			// Exclude Paths
 			ConfigSection.Add( TEXT("ExcludePathFilters"), TEXT("Content/Localization/*") );
 			for (const auto& ExcludePath : Target->Settings.GatherFromPackages.ExcludePathWildcards)
 			{
-				ConfigSection.Add( TEXT("ExcludePathFilters"), ExcludePath.Pattern );
+				ConfigSection.Add( TEXT("ExcludePathFilters"), FString::Printf(TEXT("%s%s"), *FLocalizationGatherPathRootUtil::GetResolvedPathRootToken(ExcludePath.PathRoot), *ExcludePath.Pattern) );
 			}
 
 			// Package Extensions
@@ -399,13 +400,13 @@ namespace LocalizationConfigurationScript
 			// Include Paths
 			for (const auto& IncludePath : Target->Settings.GatherFromMetaData.IncludePathWildcards)
 			{
-				ConfigSection.Add( TEXT("IncludePathFilters"), IncludePath.Pattern );
+				ConfigSection.Add( TEXT("IncludePathFilters"), FString::Printf(TEXT("%s%s"), *FLocalizationGatherPathRootUtil::GetResolvedPathRootToken(IncludePath.PathRoot), *IncludePath.Pattern) );
 			}
 
 			// Exclude Paths
 			for (const auto& ExcludePath : Target->Settings.GatherFromMetaData.ExcludePathWildcards)
 			{
-				ConfigSection.Add( TEXT("ExcludePathFilters"), ExcludePath.Pattern );
+				ConfigSection.Add( TEXT("ExcludePathFilters"), FString::Printf(TEXT("%s%s"), *FLocalizationGatherPathRootUtil::GetResolvedPathRootToken(ExcludePath.PathRoot), *ExcludePath.Pattern) );
 			}
 
 			// Package Extensions
@@ -463,7 +464,8 @@ namespace LocalizationConfigurationScript
 	{
 		FLocalizationConfigurationScript Script;
 
-		const FString ContentDirRelativeToGameDir = MakePathRelativeForCommandletProcess(GetContentDir(Target), !Target->IsMemberOfEngineTargetSet());
+		const bool bIsEngineTarget = Target->IsMemberOfEngineTargetSet();
+		const FString ContentDirRelativeToGameDir = MakePathRelativeForCommandletProcess(GetContentDir(Target), !bIsEngineTarget);
 
 		// CommonSettings
 		{
@@ -476,12 +478,12 @@ namespace LocalizationConfigurationScript
 				// The output path for a specific culture is a file path.
 				if (CultureName.IsSet())
 				{
-					SourcePath = MakePathRelativeForCommandletProcess( FPaths::GetPath(ImportPathOverride.GetValue()), !Target->IsMemberOfEngineTargetSet() );
+					SourcePath = MakePathRelativeForCommandletProcess( FPaths::GetPath(ImportPathOverride.GetValue()), !bIsEngineTarget );
 				}
 				// Otherwise, it is a directory path.
 				else
 				{
-					SourcePath = MakePathRelativeForCommandletProcess( ImportPathOverride.GetValue(), !Target->IsMemberOfEngineTargetSet() );
+					SourcePath = MakePathRelativeForCommandletProcess( ImportPathOverride.GetValue(), !bIsEngineTarget );
 				}
 			}
 			// Use the default PO file's directory path.
@@ -582,7 +584,8 @@ namespace LocalizationConfigurationScript
 	{
 		FLocalizationConfigurationScript Script;
 
-		const FString ContentDirRelativeToGameDir = MakePathRelativeForCommandletProcess(GetContentDir(Target), !Target->IsMemberOfEngineTargetSet());
+		const bool bIsEngineTarget = Target->IsMemberOfEngineTargetSet();
+		const FString ContentDirRelativeToGameDir = MakePathRelativeForCommandletProcess(GetContentDir(Target), !bIsEngineTarget);
 
 		// CommonSettings
 		{
@@ -598,12 +601,12 @@ namespace LocalizationConfigurationScript
 				// The output path for a specific culture is a file path.
 				if (CultureName.IsSet())
 				{
-					DestinationPath = MakePathRelativeForCommandletProcess( FPaths::GetPath(ExportPathOverride.GetValue()), !Target->IsMemberOfEngineTargetSet() );
+					DestinationPath = MakePathRelativeForCommandletProcess( FPaths::GetPath(ExportPathOverride.GetValue()), !bIsEngineTarget );
 				}
 				// Otherwise, it is a directory path.
 				else
 				{
-					DestinationPath = MakePathRelativeForCommandletProcess( ExportPathOverride.GetValue(), !Target->IsMemberOfEngineTargetSet() );
+					DestinationPath = MakePathRelativeForCommandletProcess( ExportPathOverride.GetValue(), !bIsEngineTarget );
 				}
 			}
 			// Use the default PO file's directory path.
@@ -705,7 +708,8 @@ namespace LocalizationConfigurationScript
 	{
 		FLocalizationConfigurationScript Script;
 
-		const FString ContentDirRelativeToGameDir = MakePathRelativeForCommandletProcess(GetContentDir(Target), !Target->IsMemberOfEngineTargetSet());
+		const bool bIsEngineTarget = Target->IsMemberOfEngineTargetSet();
+		const FString ContentDirRelativeToGameDir = MakePathRelativeForCommandletProcess(GetContentDir(Target), !bIsEngineTarget);
 
 		// CommonSettings
 		{
@@ -718,12 +722,12 @@ namespace LocalizationConfigurationScript
 				// The output path for a specific culture is a file path.
 				if (CultureName.IsSet())
 				{
-					SourcePath = MakePathRelativeForCommandletProcess(FPaths::GetPath(ImportPathOverride.GetValue()), !Target->IsMemberOfEngineTargetSet());
+					SourcePath = MakePathRelativeForCommandletProcess(FPaths::GetPath(ImportPathOverride.GetValue()), !bIsEngineTarget);
 				}
 				// Otherwise, it is a directory path.
 				else
 				{
-					SourcePath = MakePathRelativeForCommandletProcess(ImportPathOverride.GetValue(), !Target->IsMemberOfEngineTargetSet());
+					SourcePath = MakePathRelativeForCommandletProcess(ImportPathOverride.GetValue(), !bIsEngineTarget);
 				}
 			}
 			// Use the default dialogue script file's directory path.
@@ -815,7 +819,8 @@ namespace LocalizationConfigurationScript
 	{
 		FLocalizationConfigurationScript Script;
 
-		const FString ContentDirRelativeToGameDir = MakePathRelativeForCommandletProcess(GetContentDir(Target), !Target->IsMemberOfEngineTargetSet());
+		const bool bIsEngineTarget = Target->IsMemberOfEngineTargetSet();
+		const FString ContentDirRelativeToGameDir = MakePathRelativeForCommandletProcess(GetContentDir(Target), !bIsEngineTarget);
 
 		// CommonSettings
 		{
@@ -831,12 +836,12 @@ namespace LocalizationConfigurationScript
 				// The output path for a specific culture is a file path.
 				if (CultureName.IsSet())
 				{
-					DestinationPath = MakePathRelativeForCommandletProcess(FPaths::GetPath(ExportPathOverride.GetValue()), !Target->IsMemberOfEngineTargetSet());
+					DestinationPath = MakePathRelativeForCommandletProcess(FPaths::GetPath(ExportPathOverride.GetValue()), !bIsEngineTarget);
 				}
 				// Otherwise, it is a directory path.
 				else
 				{
-					DestinationPath = MakePathRelativeForCommandletProcess(ExportPathOverride.GetValue(), !Target->IsMemberOfEngineTargetSet());
+					DestinationPath = MakePathRelativeForCommandletProcess(ExportPathOverride.GetValue(), !bIsEngineTarget);
 				}
 			}
 			// Use the default dialogue script file's directory path.
@@ -926,7 +931,8 @@ namespace LocalizationConfigurationScript
 	{
 		FLocalizationConfigurationScript Script;
 
-		const FString ContentDirRelativeToGameDir = MakePathRelativeForCommandletProcess(GetContentDir(Target), !Target->IsMemberOfEngineTargetSet());
+		const bool bIsEngineTarget = Target->IsMemberOfEngineTargetSet();
+		const FString ContentDirRelativeToGameDir = MakePathRelativeForCommandletProcess(GetContentDir(Target), !bIsEngineTarget);
 
 		// CommonSettings
 		{
@@ -998,7 +1004,8 @@ namespace LocalizationConfigurationScript
 	{
 		FLocalizationConfigurationScript Script;
 
-		const FString ContentDirRelativeToGameDir = MakePathRelativeForCommandletProcess(GetContentDir(Target), !Target->IsMemberOfEngineTargetSet());
+		const bool bIsEngineTarget = Target->IsMemberOfEngineTargetSet();
+		const FString ContentDirRelativeToGameDir = MakePathRelativeForCommandletProcess(GetContentDir(Target), !bIsEngineTarget);
 
 		// CommonSettings
 		{
@@ -1044,7 +1051,8 @@ namespace LocalizationConfigurationScript
 	{
 		FLocalizationConfigurationScript Script;
 
-		const FString ContentDirRelativeToGameDir = MakePathRelativeForCommandletProcess(GetContentDir(Target), !Target->IsMemberOfEngineTargetSet());
+		const bool bIsEngineTarget = Target->IsMemberOfEngineTargetSet();
+		const FString ContentDirRelativeToGameDir = MakePathRelativeForCommandletProcess(GetContentDir(Target), !bIsEngineTarget);
 
 		// CommonSettings
 		{
@@ -1119,7 +1127,8 @@ namespace LocalizationConfigurationScript
 	{
 		FLocalizationConfigurationScript Script;
 
-		const FString ContentDirRelativeToGameDir = MakePathRelativeForCommandletProcess(GetContentDir(Target), !Target->IsMemberOfEngineTargetSet());
+		const bool bIsEngineTarget = Target->IsMemberOfEngineTargetSet();
+		const FString ContentDirRelativeToGameDir = MakePathRelativeForCommandletProcess(GetContentDir(Target), !bIsEngineTarget);
 
 		// RegenerateResources
 		{
diff --git a/Engine/Source/Developer/Localization/Private/LocalizationTargetTypes.cpp b/Engine/Source/Developer/Localization/Private/LocalizationTargetTypes.cpp
--- a/Engine/Source/Developer/Localization/Private/LocalizationTargetTypes.cpp
+++ b/Engine/Source/Developer/Localization/Private/LocalizationTargetTypes.cpp
@@ -5,6 +5,7 @@
 #include "HAL/FileManager.h"
 #include "Misc/FileHelper.h"
 #include "Misc/Paths.h"
+#include "Misc/App.h"
 #include "ISourceControlOperation.h"
 #include "SourceControlOperations.h"
 #include "ISourceControlProvider.h"
@@ -15,22 +16,79 @@
 
 #define LOCTEXT_NAMESPACE "LocalizationTargetTypes"
 
-bool FGatherTextSearchDirectory::Validate(const FString& RootDirectory, FText& OutError) const
+FString FLocalizationGatherPathRootUtil::GetResolvedPathRootToken(const ELocalizationGatherPathRoot InPathRoot)
 {
-	if (Path.IsEmpty())
+	switch (InPathRoot)
+	{
+	case ELocalizationGatherPathRoot::Engine:
+		return TEXT("%LOCENGINEROOT%");
+	case ELocalizationGatherPathRoot::Project:
+		return TEXT("%LOCPROJECTROOT%");
+	}
+
+	return FString();
+}
+
+FString FLocalizationGatherPathRootUtil::GetResolvedPathRoot(const ELocalizationGatherPathRoot InPathRoot, const bool bIsEngineTarget)
+{
+	switch (InPathRoot)
+	{
+	case ELocalizationGatherPathRoot::Auto:
+		return bIsEngineTarget ? FPaths::EngineDir() : FPaths::ProjectDir();
+	case ELocalizationGatherPathRoot::Engine:
+		return FPaths::EngineDir();
+	case ELocalizationGatherPathRoot::Project:
+		return FPaths::ProjectDir();
+	}
+
+	checkf(0, TEXT("Unknown ELocalizationGatherPathRoot!"));
+	return FString();
+}
+
+FText FLocalizationGatherPathRootUtil::GetResolvedPathRootDisplayName(const ELocalizationGatherPathRoot InPathRoot, const bool bIsEngineTarget)
+{
+	auto CreateProjectRootDisplayName = []() -> FText
+	{
+		return FText::Format(LOCTEXT("ProjectRootDisplayNameFmt", "{0}/"), FText::AsCultureInvariant(FApp::GetProjectName()));
+	};
+
+	switch (InPathRoot)
+	{
+	case ELocalizationGatherPathRoot::Auto:
+		return bIsEngineTarget ? LOCTEXT("EngineRootDisplayName", "Engine/") : CreateProjectRootDisplayName();
+	case ELocalizationGatherPathRoot::Engine:
+		return LOCTEXT("EngineRootDisplayName", "Engine/");
+	case ELocalizationGatherPathRoot::Project:
+		return CreateProjectRootDisplayName();
+	}
+
+	checkf(0, TEXT("Unknown ELocalizationGatherPathRoot!"));
+	return FText();
+}
+
+bool FGatherTextSearchDirectory::Validate(const bool bIsEngineTarget, FText& OutError) const
+{
+	if (bIsEngineTarget && PathRoot == ELocalizationGatherPathRoot::Project)
+	{
+		OutError = LOCTEXT("EngineCannotUseProjectRootError", "Engine targets cannot use 'Project' as their path root.");
+		return false;
+	}
+
+	const FString ResolvedPath = FPaths::IsRelative(Path) ? FPaths::Combine(*FLocalizationGatherPathRootUtil::GetResolvedPathRoot(PathRoot, bIsEngineTarget), *Path) : Path;
+	if (ResolvedPath.IsEmpty())
 	{
 		OutError = LOCTEXT("SearchDirectoryEmptyError", "Search directory not specified. Use \".\" to specify the root directory.");
 		return false;
 	}
 
 	FText InvalidPathReason;
-	if (!FPaths::ValidatePath(Path, &InvalidPathReason))
+	if (!FPaths::ValidatePath(ResolvedPath, &InvalidPathReason))
 	{
 		OutError = InvalidPathReason;
 		return false;
 	}
 
-	if (!FPaths::DirectoryExists(FPaths::Combine(*RootDirectory, *Path)))
+	if (!FPaths::DirectoryExists(ResolvedPath))
 	{
 		OutError = LOCTEXT("SearchDirectoryNonExistentError", "Search directory does not exist.");
 		return false;
@@ -39,7 +97,7 @@ bool FGatherTextSearchDirectory::Validate(const FString& RootDirectory, FText& O
 	return true;
 }
 
-bool FGatherTextIncludePath::Validate(const FString& RootDirectory, FText& OutError) const
+bool FGatherTextIncludePath::Validate(const bool bIsEngineTarget, FText& OutError) const
 {
 	if (Pattern.IsEmpty())
 	{
@@ -56,7 +114,7 @@ bool FGatherTextIncludePath::Validate(const FString& RootDirectory, FText& OutEr
 	return true;
 }
 
-bool FGatherTextExcludePath::Validate(FText& OutError) const
+bool FGatherTextExcludePath::Validate(const bool bIsEngineTarget, FText& OutError) const
 {
 	if (Pattern.IsEmpty())
 	{
@@ -73,7 +131,7 @@ bool FGatherTextExcludePath::Validate(FText& OutError) const
 	return true;
 }
 
-bool FGatherTextFileExtension::Validate(FText& OutError) const
+bool FGatherTextFileExtension::Validate(const bool bIsEngineTarget, FText& OutError) const
 {
 	if (Pattern.IsEmpty())
 	{
@@ -109,22 +167,22 @@ const TArray<FGatherTextFileExtension>& FGatherTextFromPackagesConfiguration::Ge
 	return DefaultPackageFileExtensions;
 }
 
-bool FGatherTextFromTextFilesConfiguration::Validate(const FString& RootDirectory, FText& OutError) const
+bool FGatherTextFromTextFilesConfiguration::Validate(const bool bIsEngineTarget, FText& OutError) const
 {
-	const auto& InvalidSearchDirectoryPredicate = [RootDirectory](const FGatherTextSearchDirectory& Element) -> bool
+	const auto& InvalidSearchDirectoryPredicate = [bIsEngineTarget](const FGatherTextSearchDirectory& Element) -> bool
 	{
 		FText ElementError;
-		return !(Element.Validate(RootDirectory, ElementError));
+		return !(Element.Validate(bIsEngineTarget, ElementError));
 	};
-	const auto& InvalidExcludePathPredicate = [RootDirectory](const FGatherTextExcludePath& Element) -> bool
+	const auto& InvalidExcludePathPredicate = [bIsEngineTarget](const FGatherTextExcludePath& Element) -> bool
 	{
 		FText ElementError;
-		return !(Element.Validate(ElementError));
+		return !(Element.Validate(bIsEngineTarget, ElementError));
 	};
-	const auto& InvalidFileExtensionPredicate = [RootDirectory](const FGatherTextFileExtension& Element) -> bool
+	const auto& InvalidFileExtensionPredicate = [bIsEngineTarget](const FGatherTextFileExtension& Element) -> bool
 	{
 		FText ElementError;
-		return !(Element.Validate(ElementError));
+		return !(Element.Validate(bIsEngineTarget, ElementError));
 	};
 
 	if (SearchDirectories.Num() > 0 && !SearchDirectories.ContainsByPredicate(InvalidSearchDirectoryPredicate) &&
@@ -138,22 +196,22 @@ bool FGatherTextFromTextFilesConfiguration::Validate(const FString& RootDirector
 	return false;
 }
 
-bool FGatherTextFromPackagesConfiguration::Validate(const FString& RootDirectory, FText& OutError) const
+bool FGatherTextFromPackagesConfiguration::Validate(const bool bIsEngineTarget, FText& OutError) const
 {
-	const auto& InvalidIncludePathPredicate = [RootDirectory](const FGatherTextIncludePath& Element) -> bool
+	const auto& InvalidIncludePathPredicate = [bIsEngineTarget](const FGatherTextIncludePath& Element) -> bool
 	{
 		FText ElementError;
-		return !(Element.Validate(RootDirectory, ElementError));
+		return !(Element.Validate(bIsEngineTarget, ElementError));
 	};
-	const auto& InvalidExcludePathPredicate = [RootDirectory](const FGatherTextExcludePath& Element) -> bool
+	const auto& InvalidExcludePathPredicate = [bIsEngineTarget](const FGatherTextExcludePath& Element) -> bool
 	{
 		FText ElementError;
-		return !(Element.Validate(ElementError));
+		return !(Element.Validate(bIsEngineTarget, ElementError));
 	};
-	const auto& InvalidFileExtensionPredicate = [RootDirectory](const FGatherTextFileExtension& Element) -> bool
+	const auto& InvalidFileExtensionPredicate = [bIsEngineTarget](const FGatherTextFileExtension& Element) -> bool
 	{
 		FText ElementError;
-		return !(Element.Validate(ElementError));
+		return !(Element.Validate(bIsEngineTarget, ElementError));
 	};
 
 	if (IncludePathWildcards.Num() > 0 && !IncludePathWildcards.ContainsByPredicate(InvalidIncludePathPredicate) &&
@@ -167,7 +225,7 @@ bool FGatherTextFromPackagesConfiguration::Validate(const FString& RootDirectory
 	return false;
 }
 
-bool FMetaDataTextKeyPattern::Validate(FText& OutError) const
+bool FMetaDataTextKeyPattern::Validate(const bool bIsEngineTarget, FText& OutError) const
 {
 	for (const auto& PossiblePlaceHolder : GetPossiblePlaceHolders())
 	{
@@ -194,7 +252,7 @@ const TArray<FString>& FMetaDataTextKeyPattern::GetPossiblePlaceHolders()
 	return PossiblePlaceHolders;
 }
 
-bool FMetaDataKeyName::Validate(FText& OutError) const
+bool FMetaDataKeyName::Validate(const bool bIsEngineTarget, FText& OutError) const
 {
 	if (Name.IsEmpty())
 	{
@@ -205,27 +263,27 @@ bool FMetaDataKeyName::Validate(FText& OutError) const
 	return true;
 }
 
-bool FMetaDataKeyGatherSpecification::Validate(FText& OutError) const
+bool FMetaDataKeyGatherSpecification::Validate(const bool bIsEngineTarget, FText& OutError) const
 {
-	return MetaDataKey.Validate(OutError) && TextKeyPattern.Validate(OutError);
+	return MetaDataKey.Validate(bIsEngineTarget, OutError) && TextKeyPattern.Validate(bIsEngineTarget, OutError);
 }
 
-bool FGatherTextFromMetaDataConfiguration::Validate(const FString& RootDirectory, FText& OutError) const
+bool FGatherTextFromMetaDataConfiguration::Validate(const bool bIsEngineTarget, FText& OutError) const
 {
-	const auto& InvalidIncludePathPredicate = [RootDirectory](const FGatherTextIncludePath& Element) -> bool
+	const auto& InvalidIncludePathPredicate = [bIsEngineTarget](const FGatherTextIncludePath& Element) -> bool
 	{
 		FText ElementError;
-		return !(Element.Validate(RootDirectory, ElementError));
+		return !(Element.Validate(bIsEngineTarget, ElementError));
 	};
-	const auto& InvalidExcludePathPredicate = [RootDirectory](const FGatherTextExcludePath& Element) -> bool
+	const auto& InvalidExcludePathPredicate = [bIsEngineTarget](const FGatherTextExcludePath& Element) -> bool
 	{
 		FText ElementError;
-		return !(Element.Validate(ElementError));
+		return !(Element.Validate(bIsEngineTarget, ElementError));
 	};
-	const auto& InvalidMetaDataKeySpecificationPredicate = [RootDirectory](const FMetaDataKeyGatherSpecification& Element) -> bool
+	const auto& InvalidMetaDataKeySpecificationPredicate = [bIsEngineTarget](const FMetaDataKeyGatherSpecification& Element) -> bool
 	{
 		FText ElementError;
-		return !(Element.Validate(ElementError));
+		return !(Element.Validate(bIsEngineTarget, ElementError));
 	};
 
 	if (IncludePathWildcards.Num() > 0 && !IncludePathWildcards.ContainsByPredicate(InvalidIncludePathPredicate) &&
diff --git a/Engine/Source/Developer/Localization/Public/LocalizationTargetTypes.h b/Engine/Source/Developer/Localization/Public/LocalizationTargetTypes.h
--- a/Engine/Source/Developer/Localization/Public/LocalizationTargetTypes.h
+++ b/Engine/Source/Developer/Localization/Public/LocalizationTargetTypes.h
@@ -23,15 +23,41 @@ enum class ELocalizationTargetConflictStatus : uint8
 	Clear,
 };
 
+UENUM()
+enum class ELocalizationGatherPathRoot : uint8
+{
+	/** Automatically select Engine or Project based on which set the target belongs to */
+	Auto,
+	/** Path is relative to the Engine directory */
+	Engine,
+	/** Path is relative to the Project directory */
+	Project,
+};
+
+struct LOCALIZATION_API FLocalizationGatherPathRootUtil
+{
+	/** Get the resolved placeholder token of the path root ("%LOCENGINEROOT%" for Engine or "%LOCPROJECTROOT%" for Project; empty for Auto) */
+	static FString GetResolvedPathRootToken(const ELocalizationGatherPathRoot InPathRoot);
+
+	/** Get the path name of the resolved path root (either FPaths::EngineDir() or FPaths::ProjectDir() depending on the setting and target type */
+	static FString GetResolvedPathRoot(const ELocalizationGatherPathRoot InPathRoot, const bool bIsEngineTarget);
+
+	/** Get the display name of the resolved path root (either "Engine/" or "{ProjectName/" depending on the setting and target type */
+	static FText GetResolvedPathRootDisplayName(const ELocalizationGatherPathRoot InPathRoot, const bool bIsEngineTarget);
+};
+
 USTRUCT()
 struct FGatherTextSearchDirectory
 {
 	GENERATED_USTRUCT_BODY()
 
+	UPROPERTY(config, EditAnywhere, Category="Path")
+	ELocalizationGatherPathRoot PathRoot;
+
 	UPROPERTY(config, EditAnywhere, Category="Path")
 	FString Path;
 
-	LOCALIZATION_API bool Validate(const FString& RootDirectory, FText& OutError) const;
+	LOCALIZATION_API bool Validate(const bool bIsEngineTarget, FText& OutError) const;
 };
 
 USTRUCT()
@@ -39,10 +65,13 @@ struct FGatherTextIncludePath
 {
 	GENERATED_USTRUCT_BODY()
 
+	UPROPERTY(config, EditAnywhere, Category="Path")
+	ELocalizationGatherPathRoot PathRoot;
+
 	UPROPERTY(config, EditAnywhere, Category="Pattern")
 	FString Pattern;
 
-	LOCALIZATION_API bool Validate(const FString& RootDirectory, FText& OutError) const;
+	LOCALIZATION_API bool Validate(const bool bIsEngineTarget, FText& OutError) const;
 };
 
 USTRUCT()
@@ -50,10 +79,13 @@ struct FGatherTextExcludePath
 {
 	GENERATED_USTRUCT_BODY()
 
+	UPROPERTY(config, EditAnywhere, Category="Path")
+	ELocalizationGatherPathRoot PathRoot;
+
 	UPROPERTY(config, EditAnywhere, Category="Pattern")
 	FString Pattern;
 
-	LOCALIZATION_API bool Validate(FText& OutError) const;
+	LOCALIZATION_API bool Validate(const bool bIsEngineTarget, FText& OutError) const;
 };
 
 USTRUCT()
@@ -64,7 +96,7 @@ struct FGatherTextFileExtension
 	UPROPERTY(config, EditAnywhere, Category="Pattern")
 	FString Pattern;
 
-	LOCALIZATION_API bool Validate(FText& OutError) const;
+	LOCALIZATION_API bool Validate(const bool bIsEngineTarget, FText& OutError) const;
 };
 
 USTRUCT()
@@ -101,7 +133,7 @@ struct FGatherTextFromTextFilesConfiguration
 	UPROPERTY(config, EditAnywhere, Category = "Filter")
 	bool ShouldGatherFromEditorOnlyData;
 
-	LOCALIZATION_API bool Validate(const FString& RootDirectory, FText& OutError) const;
+	LOCALIZATION_API bool Validate(const bool bIsEngineTarget, FText& OutError) const;
 };
 
 
@@ -153,7 +185,7 @@ struct FGatherTextFromPackagesConfiguration
 	UPROPERTY(config, EditAnywhere, Category = "Gather Text", AdvancedDisplay)
 	bool SkipGatherCache;
 
-	LOCALIZATION_API bool Validate(const FString& RootDirectory, FText& OutError) const;
+	LOCALIZATION_API bool Validate(const bool bIsEngineTarget, FText& OutError) const;
 };
 
 USTRUCT()
@@ -164,7 +196,7 @@ struct FMetaDataTextKeyPattern
 	UPROPERTY(config, EditAnywhere, Category="Pattern")
 	FString Pattern;
 
-	LOCALIZATION_API bool Validate(FText& OutError) const;
+	LOCALIZATION_API bool Validate(const bool bIsEngineTarget, FText& OutError) const;
 
 	LOCALIZATION_API static const TArray<FString>& GetPossiblePlaceHolders();
 };
@@ -177,7 +209,7 @@ struct FMetaDataKeyName
 	UPROPERTY(config, EditAnywhere, Category="Name")
 	FString Name;
 
-	LOCALIZATION_API bool Validate(FText& OutError) const;
+	LOCALIZATION_API bool Validate(const bool bIsEngineTarget, FText& OutError) const;
 };
 
 USTRUCT()
@@ -200,7 +232,7 @@ struct FMetaDataKeyGatherSpecification
 	UPROPERTY(config, EditAnywhere, Category = "Output")
 	FMetaDataTextKeyPattern TextKeyPattern;
 
-	LOCALIZATION_API bool Validate(FText& OutError) const;
+	LOCALIZATION_API bool Validate(const bool bIsEngineTarget, FText& OutError) const;
 };
 
 USTRUCT()
@@ -234,7 +266,7 @@ struct FGatherTextFromMetaDataConfiguration
 	UPROPERTY(config, EditAnywhere, Category = "Filter")
 	bool ShouldGatherFromEditorOnlyData;
 
-	LOCALIZATION_API bool Validate(const FString& RootDirectory, FText& OutError) const;
+	LOCALIZATION_API bool Validate(const bool bIsEngineTarget, FText& OutError) const;
 };
 
 USTRUCT()
diff --git a/Engine/Source/Editor/LocalizationDashboard/Private/GatherTextDetailCustomizations.cpp b/Engine/Source/Editor/LocalizationDashboard/Private/GatherTextDetailCustomizations.cpp
--- a/Engine/Source/Editor/LocalizationDashboard/Private/GatherTextDetailCustomizations.cpp
+++ b/Engine/Source/Editor/LocalizationDashboard/Private/GatherTextDetailCustomizations.cpp
@@ -80,36 +80,53 @@ namespace
 		SLATE_END_ARGS()
 
 	public:
-		void Construct(const FArguments& Args, const TSharedRef<IPropertyHandle>& PathStringPropertyHandle);
+		void Construct(const FArguments& Args, const TSharedRef<IPropertyHandle>& InPathRootPropertyHandle, const TSharedRef<IPropertyHandle>& InPathStringPropertyHandle);
 
 	private:
 		FReply PathPickerOnClicked();
+		ULocalizationTarget* GetLocalizationTarget() const;
+		ELocalizationGatherPathRoot GetPathRoot() const;
+		FString GetResolvedPathRoot() const;
+		FText GetResolvedPathRootDisplayName() const;
 
 	private:
+		TSharedPtr<IPropertyHandle> PathRootPropertyHandle;
 		TSharedPtr<IPropertyHandle> PathStringPropertyHandle;
 		bool ShouldCoercePathAsWildcardPattern;
 	};
 
-	void SGatherTextPathPicker::Construct(const FArguments& Args, const TSharedRef<IPropertyHandle>& InPathStringPropertyHandle)
+	void SGatherTextPathPicker::Construct(const FArguments& Args, const TSharedRef<IPropertyHandle>& InPathRootPropertyHandle, const TSharedRef<IPropertyHandle>& InPathStringPropertyHandle)
 	{
+		PathRootPropertyHandle = InPathRootPropertyHandle;
 		PathStringPropertyHandle = InPathStringPropertyHandle;
 		ShouldCoercePathAsWildcardPattern = Args._ShouldCoercePathAsWildcardPattern;
 
-		TArray<UObject*> OuterObjects;
-		PathStringPropertyHandle->GetOuterObjects(OuterObjects);
-		ULocalizationTarget* const LocalizationTarget = Cast<ULocalizationTarget>(OuterObjects.Top());
-			
 		ChildSlot
 			[
 				SNew(SHorizontalBox)
 				+SHorizontalBox::Slot()
 				.AutoWidth()
 				.VAlign(VAlign_Center)
+				[
+					PathRootPropertyHandle.IsValid() && PathRootPropertyHandle->IsValidHandle() ? PathRootPropertyHandle->CreatePropertyValueWidget() : SNullWidget::NullWidget
+				]
+				+SHorizontalBox::Slot()
+				.AutoWidth()
+				.VAlign(VAlign_Center)
 				.Padding(2.0f, 0.0f)
 				[
 					SNew(STextBlock)
 					.Font(Args._Font)
-					.Text( FText::FromString( FString::Printf( TEXT("%s/"), *(FPaths::GetBaseFilename(LocalizationTarget->IsMemberOfEngineTargetSet() ? FPaths::EngineDir() : FPaths::ProjectDir())) ) ) )
+					.Text(LOCTEXT("Colon", ":"))
+				]
+				+SHorizontalBox::Slot()
+				.AutoWidth()
+				.VAlign(VAlign_Center)
+				.Padding(2.0f, 0.0f)
+				[
+					SNew(STextBlock)
+					.Font(Args._Font)
+					.Text(this, &SGatherTextPathPicker::GetResolvedPathRootDisplayName)
 				]
 				+SHorizontalBox::Slot()
 				.AutoWidth()
@@ -152,10 +169,7 @@ namespace
 		IDesktopPlatform* const DesktopPlatform = FDesktopPlatformModule::Get();
 		check(DesktopPlatform);
 
-		TArray<UObject*> OuterObjects;
-		PathStringPropertyHandle->GetOuterObjects(OuterObjects);
-		ULocalizationTarget* const LocalizationTarget = CastChecked<ULocalizationTarget>(OuterObjects.Top());
-		const FString DesiredRootPath = FPaths::ConvertRelativePathToFull(LocalizationTarget->IsMemberOfEngineTargetSet() ? FPaths::EngineDir() : FPaths::ProjectDir());
+		const FString DesiredRootPath = FPaths::ConvertRelativePathToFull(GetResolvedPathRoot());
 
 		const TSharedPtr<SWindow> ParentWindow = FSlateApplication::Get().FindWidgetWindow(AsShared());
 		const TSharedPtr<FGenericWindow> ParentGenericWindow = ParentWindow.IsValid() ? ParentWindow->GetNativeWindow() : nullptr;
@@ -194,11 +208,48 @@ namespace
 
 		return FReply::Handled();
 	}
+
+	ULocalizationTarget* SGatherTextPathPicker::GetLocalizationTarget() const
+	{
+		if (!PathStringPropertyHandle.IsValid() || !PathStringPropertyHandle->IsValidHandle())
+		{
+			return nullptr;
+		}
+
+		TArray<UObject*> OuterObjects;
+		PathStringPropertyHandle->GetOuterObjects(OuterObjects);
+		return CastChecked<ULocalizationTarget>(OuterObjects.Top());
+	}
+
+	ELocalizationGatherPathRoot SGatherTextPathPicker::GetPathRoot() const
+	{
+		if (!PathRootPropertyHandle.IsValid() || !PathRootPropertyHandle->IsValidHandle())
+		{
+			return ELocalizationGatherPathRoot::Auto;
+		}
+
+		TArray<const void*> RawData;
+		PathRootPropertyHandle->AccessRawData(RawData);
+		return *(ELocalizationGatherPathRoot*)RawData.Top();
+	}
+
+	FString SGatherTextPathPicker::GetResolvedPathRoot() const
+	{
+		ULocalizationTarget* LocalizationTarget = GetLocalizationTarget();
+		return FLocalizationGatherPathRootUtil::GetResolvedPathRoot(GetPathRoot(), LocalizationTarget && LocalizationTarget->IsMemberOfEngineTargetSet());
+	}
+
+	FText SGatherTextPathPicker::GetResolvedPathRootDisplayName() const
+	{
+		ULocalizationTarget* LocalizationTarget = GetLocalizationTarget();
+		return FLocalizationGatherPathRootUtil::GetResolvedPathRootDisplayName(GetPathRoot(), LocalizationTarget && LocalizationTarget->IsMemberOfEngineTargetSet());
+	}
 }
 
 void FGatherTextSearchDirectoryStructCustomization::CustomizeStructHeader( TSharedRef<IPropertyHandle> StructPropertyHandle, class FDetailWidgetRow& HeaderRow, IStructCustomizationUtils& StructCustomizationUtils )
 {
-	const TSharedPtr<IPropertyHandle> PathPropertyHandle = StructPropertyHandle->GetChildHandle(GET_MEMBER_NAME_CHECKED(FGatherTextSearchDirectory, Path));
+	const TSharedPtr<IPropertyHandle> PathRootPropertyHandle = StructPropertyHandle->GetChildHandle(GET_MEMBER_NAME_CHECKED(FGatherTextSearchDirectory, PathRoot));
+	const TSharedPtr<IPropertyHandle> PathStringPropertyHandle = StructPropertyHandle->GetChildHandle(GET_MEMBER_NAME_CHECKED(FGatherTextSearchDirectory, Path));
 
 	const auto& GetErrorText = [StructPropertyHandle]() -> FText
 	{
@@ -213,9 +264,8 @@ void FGatherTextSearchDirectoryStructCustomization::CustomizeStructHeader( TShar
 			ULocalizationTarget* const LocalizationTarget = Cast<ULocalizationTarget>(OuterObjects.Top());
 			if (LocalizationTarget)
 			{
-				const FString DesiredRootPath = FPaths::ConvertRelativePathToFull(LocalizationTarget->IsMemberOfEngineTargetSet() ? FPaths::EngineDir() : FPaths::ProjectDir());
 				FText Error;
-				if (!Struct.Validate(DesiredRootPath, Error))
+				if (!Struct.Validate(LocalizationTarget->IsMemberOfEngineTargetSet(), Error))
 				{
 					return Error;
 				}
@@ -244,13 +294,14 @@ void FGatherTextSearchDirectoryStructCustomization::CustomizeStructHeader( TShar
 		.ValueContent()
 		.MaxDesiredWidth(TOptional<float>())
 		[
-			SNew(SGatherTextPathPicker, PathPropertyHandle.ToSharedRef())
+			SNew(SGatherTextPathPicker, PathRootPropertyHandle.ToSharedRef(), PathStringPropertyHandle.ToSharedRef())
 			.Font(StructCustomizationUtils.GetRegularFont())
 		];
 }
 
 void FGatherTextIncludePathStructCustomization::CustomizeStructHeader( TSharedRef<IPropertyHandle> StructPropertyHandle, class FDetailWidgetRow& HeaderRow, IStructCustomizationUtils& StructCustomizationUtils )
 {
+	const TSharedPtr<IPropertyHandle> PathRootPropertyHandle = StructPropertyHandle->GetChildHandle(GET_MEMBER_NAME_CHECKED(FGatherTextIncludePath, PathRoot));
 	const TSharedPtr<IPropertyHandle> PatternPropertyHandle = StructPropertyHandle->GetChildHandle(GET_MEMBER_NAME_CHECKED(FGatherTextIncludePath, Pattern));
 
 	const auto& GetErrorText = [StructPropertyHandle]() -> FText
@@ -266,9 +317,8 @@ void FGatherTextIncludePathStructCustomization::CustomizeStructHeader( TSharedRe
 			ULocalizationTarget* const LocalizationTarget = Cast<ULocalizationTarget>(OuterObjects.Top());
 			if (LocalizationTarget)
 			{
-				const FString DesiredRootPath = FPaths::ConvertRelativePathToFull(LocalizationTarget->IsMemberOfEngineTargetSet() ? FPaths::EngineDir() : FPaths::ProjectDir());
 				FText Error;
-				if (!Struct.Validate(DesiredRootPath, Error))
+				if (!Struct.Validate(LocalizationTarget->IsMemberOfEngineTargetSet(), Error))
 				{
 					return Error;
 				}
@@ -297,7 +347,7 @@ void FGatherTextIncludePathStructCustomization::CustomizeStructHeader( TSharedRe
 		.ValueContent()
 		.MaxDesiredWidth(TOptional<float>())
 		[
-			SNew(SGatherTextPathPicker, PatternPropertyHandle.ToSharedRef())
+			SNew(SGatherTextPathPicker, PathRootPropertyHandle.ToSharedRef(), PatternPropertyHandle.ToSharedRef())
 			.Font(StructCustomizationUtils.GetRegularFont())
 			.ShouldCoercePathAsWildcardPattern(true)
 		];
@@ -305,6 +355,7 @@ void FGatherTextIncludePathStructCustomization::CustomizeStructHeader( TSharedRe
 
 void FGatherTextExcludePathStructCustomization::CustomizeStructHeader( TSharedRef<IPropertyHandle> StructPropertyHandle, class FDetailWidgetRow& HeaderRow, IStructCustomizationUtils& StructCustomizationUtils )
 {
+	const TSharedPtr<IPropertyHandle> PathRootPropertyHandle = StructPropertyHandle->GetChildHandle(GET_MEMBER_NAME_CHECKED(FGatherTextExcludePath, PathRoot));
 	const TSharedPtr<IPropertyHandle> PatternPropertyHandle = StructPropertyHandle->GetChildHandle(GET_MEMBER_NAME_CHECKED(FGatherTextExcludePath, Pattern));
 
 	const auto& GetErrorText = [StructPropertyHandle]() -> FText
@@ -315,10 +366,16 @@ void FGatherTextExcludePathStructCustomization::CustomizeStructHeader( TSharedRe
 			StructPropertyHandle->AccessRawData(RawData);
 			const FGatherTextExcludePath& Struct = *(reinterpret_cast<const FGatherTextExcludePath*>(RawData.Top()));
 
-			FText Error;
-			if (!Struct.Validate(Error))
+			TArray<UObject*> OuterObjects;
+			StructPropertyHandle->GetOuterObjects(OuterObjects);
+			ULocalizationTarget* const LocalizationTarget = Cast<ULocalizationTarget>(OuterObjects.Top());
+			if (LocalizationTarget)
 			{
-				return Error;
+				FText Error;
+				if (!Struct.Validate(LocalizationTarget->IsMemberOfEngineTargetSet(), Error))
+				{
+					return Error;
+				}
 			}
 		}
 
@@ -344,7 +401,7 @@ void FGatherTextExcludePathStructCustomization::CustomizeStructHeader( TSharedRe
 		.ValueContent()
 		.MaxDesiredWidth(TOptional<float>())
 		[
-			SNew(SGatherTextPathPicker, PatternPropertyHandle.ToSharedRef())
+			SNew(SGatherTextPathPicker, PathRootPropertyHandle.ToSharedRef(), PatternPropertyHandle.ToSharedRef())
 			.Font(StructCustomizationUtils.GetRegularFont())
 			.ShouldCoercePathAsWildcardPattern(true)
 		];
@@ -362,10 +419,16 @@ void FGatherTextFileExtensionStructCustomization::CustomizeStructHeader( TShared
 			StructPropertyHandle->AccessRawData(RawData);
 			const FGatherTextFileExtension& Struct = *(reinterpret_cast<const FGatherTextFileExtension*>(RawData.Top()));
 
-			FText Error;
-			if (!Struct.Validate(Error))
+			TArray<UObject*> OuterObjects;
+			StructPropertyHandle->GetOuterObjects(OuterObjects);
+			ULocalizationTarget* const LocalizationTarget = Cast<ULocalizationTarget>(OuterObjects.Top());
+			if (LocalizationTarget)
 			{
-				return Error;
+				FText Error;
+				if (!Struct.Validate(LocalizationTarget->IsMemberOfEngineTargetSet(), Error))
+				{
+					return Error;
+				}
 			}
 		}
 
@@ -475,9 +538,8 @@ void FGatherTextFromTextFilesConfigurationStructCustomization::CustomizeStructHe
 		ULocalizationTarget* const LocalizationTarget = Cast<ULocalizationTarget>(OuterObjects.Top());
 		if (LocalizationTarget)
 		{
-			const FString DesiredRootPath = FPaths::ConvertRelativePathToFull(LocalizationTarget->IsMemberOfEngineTargetSet() ? FPaths::EngineDir() : FPaths::ProjectDir());
 			FText Error;
-			if (!Struct.Validate(DesiredRootPath, Error))
+			if (!Struct.Validate(LocalizationTarget->IsMemberOfEngineTargetSet(), Error))
 			{
 				return Error;
 			}
@@ -566,9 +628,8 @@ void FGatherTextFromPackagesConfigurationStructCustomization::CustomizeStructHea
 		ULocalizationTarget* const LocalizationTarget = Cast<ULocalizationTarget>(OuterObjects.Top());
 		if (LocalizationTarget)
 		{
-			const FString DesiredRootPath = FPaths::ConvertRelativePathToFull(LocalizationTarget->IsMemberOfEngineTargetSet() ? FPaths::EngineDir() : FPaths::ProjectDir());
 			FText Error;
-			if (!Struct.Validate(DesiredRootPath, Error))
+			if (!Struct.Validate(LocalizationTarget->IsMemberOfEngineTargetSet(), Error))
 			{
 				return Error;
 			}
@@ -771,10 +832,16 @@ void FMetaDataTextKeyPatternStructCustomization::CustomizeStructHeader( TSharedR
 			StructPropertyHandle->AccessRawData(RawData);
 			const FMetaDataTextKeyPattern& Struct = *(reinterpret_cast<const FMetaDataTextKeyPattern*>(RawData.Top()));
 
-			FText Error;
-			if (!Struct.Validate(Error))
+			TArray<UObject*> OuterObjects;
+			StructPropertyHandle->GetOuterObjects(OuterObjects);
+			ULocalizationTarget* const LocalizationTarget = Cast<ULocalizationTarget>(OuterObjects.Top());
+			if (LocalizationTarget)
 			{
-				return Error;
+				FText Error;
+				if (!Struct.Validate(LocalizationTarget->IsMemberOfEngineTargetSet(), Error))
+				{
+					return Error;
+				}
 			}
 		}
 
@@ -816,10 +883,16 @@ void FMetaDataKeyNameStructCustomization::CustomizeStructHeader( TSharedRef<IPro
 			StructPropertyHandle->AccessRawData(RawData);
 			const FMetaDataKeyName& Struct = *(reinterpret_cast<const FMetaDataKeyName*>(RawData.Top()));
 
-			FText Error;
-			if (!Struct.Validate(Error))
+			TArray<UObject*> OuterObjects;
+			StructPropertyHandle->GetOuterObjects(OuterObjects);
+			ULocalizationTarget* const LocalizationTarget = Cast<ULocalizationTarget>(OuterObjects.Top());
+			if (LocalizationTarget)
 			{
-				return Error;
+				FText Error;
+				if (!Struct.Validate(LocalizationTarget->IsMemberOfEngineTargetSet(), Error))
+				{
+					return Error;
+				}
 			}
 		}
 
@@ -858,10 +931,16 @@ void FMetaDataKeyGatherSpecificationStructCustomization::CustomizeStructHeader(T
 			StructPropertyHandle->AccessRawData(RawData);
 			const FMetaDataKeyGatherSpecification& Struct = *(reinterpret_cast<const FMetaDataKeyGatherSpecification*>(RawData.Top()));
 
-			FText Error;
-			if (!Struct.Validate(Error))
+			TArray<UObject*> OuterObjects;
+			StructPropertyHandle->GetOuterObjects(OuterObjects);
+			ULocalizationTarget* const LocalizationTarget = Cast<ULocalizationTarget>(OuterObjects.Top());
+			if (LocalizationTarget)
 			{
-				return Error;
+				FText Error;
+				if (!Struct.Validate(LocalizationTarget->IsMemberOfEngineTargetSet(), Error))
+				{
+					return Error;
+				}
 			}
 		}
 
@@ -924,9 +1003,8 @@ void FGatherTextFromMetaDataConfigurationStructCustomization::CustomizeStructHea
 		ULocalizationTarget* const LocalizationTarget = Cast<ULocalizationTarget>(OuterObjects.Top());
 		if (LocalizationTarget)
 		{
-			const FString DesiredRootPath = FPaths::ConvertRelativePathToFull(LocalizationTarget->IsMemberOfEngineTargetSet() ? FPaths::EngineDir() : FPaths::ProjectDir());
 			FText Error;
-			if (!Struct.Validate(DesiredRootPath, Error))
+			if (!Struct.Validate(LocalizationTarget->IsMemberOfEngineTargetSet(), Error))
 			{
 				return Error;
 			}
diff --git a/Engine/Source/Editor/UnrealEd/Private/Commandlets/GatherTextCommandletBase.cpp b/Engine/Source/Editor/UnrealEd/Private/Commandlets/GatherTextCommandletBase.cpp
--- a/Engine/Source/Editor/UnrealEd/Private/Commandlets/GatherTextCommandletBase.cpp
+++ b/Engine/Source/Editor/UnrealEd/Private/Commandlets/GatherTextCommandletBase.cpp
@@ -96,23 +96,29 @@ bool UGatherTextCommandletBase::GetStringFromConfig( const TCHAR* Section, const
 	return bSuccess;
 }
 
-bool UGatherTextCommandletBase::GetPathFromConfig( const TCHAR* Section, const TCHAR* Key, FString& OutValue, const FString& Filename )
+void ResolveLocalizationPath(FString& InOutPath)
 {
-	bool bSuccess = GetStringFromConfig( Section, Key, OutValue, Filename );
+	static const bool bIsEngineTarget = FPaths::ProjectDir().IsEmpty();
+	static const FString AbsoluteEnginePath = FPaths::ConvertRelativePathToFull(FPaths::EngineDir()) / FString();
+	static const FString AbsoluteProjectPath = FPaths::ConvertRelativePathToFull(FPaths::ProjectDir()) / FString();
+
+	InOutPath.ReplaceInline(TEXT("%LOCENGINEROOT%"), *AbsoluteEnginePath, ESearchCase::CaseSensitive);
+	InOutPath.ReplaceInline(TEXT("%LOCPROJECTROOT%"), *AbsoluteProjectPath, ESearchCase::CaseSensitive);
 
-	if( bSuccess )
+	if (FPaths::IsRelative(InOutPath))
 	{
-		if (FPaths::IsRelative(OutValue))
-		{
-			if (!FPaths::ProjectDir().IsEmpty())
-			{
-				OutValue = FPaths::Combine( *( FPaths::ProjectDir() ), *OutValue );
-			}
-			else
-			{
-				OutValue = FPaths::Combine( *( FPaths::EngineDir() ), *OutValue );
-			}
-		}
+		InOutPath.InsertAt(0, bIsEngineTarget ? AbsoluteEnginePath : AbsoluteProjectPath);
+	}
+
+	FPaths::CollapseRelativeDirectories(InOutPath);
+}
+
+bool UGatherTextCommandletBase::GetPathFromConfig( const TCHAR* Section, const TCHAR* Key, FString& OutValue, const FString& Filename )
+{
+	const bool bSuccess = GetStringFromConfig( Section, Key, OutValue, Filename );
+	if (bSuccess)
+	{
+		ResolveLocalizationPath(OutValue);
 	}
 	return bSuccess;
 }
@@ -132,16 +138,11 @@ int32 UGatherTextCommandletBase::GetPathArrayFromConfig( const TCHAR* Section, c
 {
 	int32 count = GetStringArrayFromConfig( Section, Key, OutArr, Filename );
 
-	for (int32 i = 0; i < count; ++i)
+	for (FString& Path : OutArr)
 	{
-		if (FPaths::IsRelative(OutArr[i]))
-		{
-			const FString ProjectBasePath = FPaths::ProjectDir().IsEmpty() ? FPaths::EngineDir() : FPaths::ProjectDir();
-			OutArr[i] = FPaths::Combine( *ProjectBasePath, *OutArr[i] );
-			OutArr[i] = FPaths::ConvertRelativePathToFull(OutArr[i]);
-		}
-		FPaths::CollapseRelativeDirectories(OutArr[i]);
+		ResolveLocalizationPath(Path);
 	}
+
 	return count;
 }
 
