Add a scrollbox to the localization dashboard

diff --git a/Engine/Source/Editor/LocalizationDashboard/Private/LocalizationDashboard.cpp b/Engine/Source/Editor/LocalizationDashboard/Private/LocalizationDashboard.cpp
--- a/Engine/Source/Editor/LocalizationDashboard/Private/LocalizationDashboard.cpp
+++ b/Engine/Source/Editor/LocalizationDashboard/Private/LocalizationDashboard.cpp
@@ -7,6 +7,7 @@
 #include "Widgets/SWidget.h"
 #include "Widgets/SCompoundWidget.h"
 #include "Widgets/SBoxPanel.h"
+#include "Widgets/Layout/SScrollBox.h"
 #include "Framework/MultiBox/MultiBoxExtender.h"
 #include "Framework/Docking/WorkspaceItem.h"
 #include "Framework/Docking/TabManager.h"
@@ -71,7 +72,13 @@ void SLocalizationDashboard::Construct(const FArguments& InArgs, const TSharedPt
 			.Label(LOCTEXT("TargetsTabLabel", "Targets"));
 
 		TSharedRef<SVerticalBox> VBox = SNew(SVerticalBox);
-		DockTab->SetContent(VBox);
+		TSharedRef<SScrollBox> ScrollBox = SNew(SScrollBox)
+			+ SScrollBox::Slot()
+			[
+				VBox
+			];
+
+		DockTab->SetContent(ScrollBox);
 
 		FPropertyEditorModule& PropertyModule = FModuleManager::LoadModuleChecked<FPropertyEditorModule>("PropertyEditor");
 
