Fix LiveLinkHub linux build

diff --git a/Engine/Source/Editor/LocalizationDashboard/Private/LocalizationDashboardSettingsDetailCustomization.cpp b/Engine/Source/Editor/LocalizationDashboard/Private/LocalizationDashboardSettingsDetailCustomization.cpp
--- a/Engine/Source/Editor/LocalizationDashboard/Private/LocalizationDashboardSettingsDetailCustomization.cpp
+++ b/Engine/Source/Editor/LocalizationDashboard/Private/LocalizationDashboardSettingsDetailCustomization.cpp
@@ -91,11 +91,13 @@ void FLocalizationDashboardSettingsDetailCustomization::CustomizeDetails(IDetail
 				]
 			];
 
+#if LOCALIZATION_SERVICES_WITH_SLATE
 		const ILocalizationServiceProvider& LSP = ILocalizationServiceModule::Get().GetProvider();
 		if (ServiceProviderCategoryBuilder != nullptr)
 		{
 			LSP.CustomizeSettingsDetails(*ServiceProviderCategoryBuilder);
 		}
+#endif
 	}
 
 	// Source Control
@@ -187,10 +189,12 @@ void FLocalizationDashboardSettingsDetailCustomization::ServiceProviderComboBox_
 	FName ServiceProviderName = LSP ? LSP->GetName() : FName(TEXT("None"));
 	ILocalizationServiceModule::Get().SetProvider(ServiceProviderName);
 
+#if LOCALIZATION_SERVICES_WITH_SLATE
 	if (LSP && ServiceProviderCategoryBuilder)
 	{
 		LSP->CustomizeSettingsDetails(*ServiceProviderCategoryBuilder);
 	}
+#endif
 	DetailLayoutBuilder->ForceRefreshDetails();
 }
 
diff --git a/Engine/Source/Editor/LocalizationDashboard/Private/LocalizationTargetDetailCustomization.cpp b/Engine/Source/Editor/LocalizationDashboard/Private/LocalizationTargetDetailCustomization.cpp
--- a/Engine/Source/Editor/LocalizationDashboard/Private/LocalizationTargetDetailCustomization.cpp
+++ b/Engine/Source/Editor/LocalizationDashboard/Private/LocalizationTargetDetailCustomization.cpp
@@ -278,10 +278,12 @@ void FLocalizationTargetDetailCustomization::CustomizeDetails(IDetailLayoutBuild
 		const TSharedRef< FUICommandList > CommandList = MakeShareable(new FUICommandList);
 		// Let the localization service extend this toolbar
 		TSharedRef<FExtender> LocalizationServiceExtender = MakeShareable(new FExtender);
+#if LOCALIZATION_SERVICES_WITH_SLATE
 		if (LocalizationTarget.IsValid() && ILocalizationServiceModule::Get().IsEnabled())
 		{
 			LSP.CustomizeTargetToolbar(LocalizationServiceExtender, LocalizationTarget);
 		}
+#endif
 		FToolBarBuilder ToolBarBuilder(CommandList, FMultiBoxCustomization::AllowCustomization("LocalizationTargetEditor"), LocalizationServiceExtender);
 
 		TAttribute<FText> GatherToolTipTextAttribute = TAttribute<FText>::Create(TAttribute<FText>::FGetter::CreateLambda([this]() -> FText
diff --git a/Engine/Source/Editor/LocalizationDashboard/Private/LocalizationTargetSetDetailCustomization.cpp b/Engine/Source/Editor/LocalizationDashboard/Private/LocalizationTargetSetDetailCustomization.cpp
--- a/Engine/Source/Editor/LocalizationDashboard/Private/LocalizationTargetSetDetailCustomization.cpp
+++ b/Engine/Source/Editor/LocalizationDashboard/Private/LocalizationTargetSetDetailCustomization.cpp
@@ -145,10 +145,12 @@ void FLocalizationTargetSetDetailCustomization::CustomizeDetails(IDetailLayoutBu
 
 			// Let the localization service extend this toolbar
 			TSharedRef<FExtender> LocalizationServiceExtender = MakeShareable(new FExtender);
+#if LOCALIZATION_SERVICES_WITH_SLATE
 			if (TargetSet.IsValid() && ILocalizationServiceModule::Get().IsEnabled())
 			{
 				LSP.CustomizeTargetSetToolbar(LocalizationServiceExtender, TargetSet);
 			}
+#endif
 
 			FToolBarBuilder ToolBarBuilder( CommandList, FMultiBoxCustomization::AllowCustomization("LocalizationDashboard"), LocalizationServiceExtender, /*bForceSmallIcons*/true );
 
