Fix issue where NativeCultureIndex wouldn't be updated when deleting a culture, causing the native culture to be modified incorrectly.

diff --git a/Engine/Source/Editor/LocalizationDashboard/Private/SLocalizationTargetEditorCultureRow.cpp b/Engine/Source/Editor/LocalizationDashboard/Private/SLocalizationTargetEditorCultureRow.cpp
--- a/Engine/Source/Editor/LocalizationDashboard/Private/SLocalizationTargetEditorCultureRow.cpp
+++ b/Engine/Source/Editor/LocalizationDashboard/Private/SLocalizationTargetEditorCultureRow.cpp
@@ -715,6 +715,16 @@ void SLocalizationTargetEditorCultureRow::Delete()
 						const TSharedPtr<IPropertyHandleArray> ArrayPropertyHandle = SupportedCulturesStatisticsPropertyHandle->AsArray();
 						if (ArrayPropertyHandle.IsValid())
 						{
+							// Update the NativeCultureIndex if it needs to be changed based on the array resize
+							// We do this before deleting the item so that SLocalizationTargetEditor::OnFinishedChangingProperties runs with the proper index
+							FLocalizationTargetSettings* const TargetSettings = GetTargetSettings();
+							if (CultureIndex < TargetSettings->NativeCultureIndex)
+							{
+								TargetSettingsPropertyHandle->NotifyPreChange();
+								TargetSettings->NativeCultureIndex--;
+								TargetSettingsPropertyHandle->NotifyPostChange(EPropertyChangeType::ValueSet);
+							}
+
 							ArrayPropertyHandle->DeleteItem(CultureIndex);
 						}
 					}
