Rename Culture to Language in the Localization Dashboard

#rnx

diff --git a/Engine/Source/Developer/Localization/Public/LocalizationTargetTypes.h b/Engine/Source/Developer/Localization/Public/LocalizationTargetTypes.h
--- a/Engine/Source/Developer/Localization/Public/LocalizationTargetTypes.h
+++ b/Engine/Source/Developer/Localization/Public/LocalizationTargetTypes.h
@@ -350,7 +350,7 @@ struct FLocalizationCompilationSettings
 	UPROPERTY(config, EditAnywhere, Category="Source")
 	bool SkipSourceCheck;
 
-	/* Should we validate that format patterns are valid for the culture being compiled (eg, detect invalid plural rules or broken syntax). */
+	/* Should we validate that format patterns are valid for the language being compiled (eg, detect invalid plural rules or broken syntax). */
 	UPROPERTY(config, EditAnywhere, Category="Validation")
 	bool ValidateFormatPatterns;
 
@@ -370,15 +370,15 @@ struct FLocalizationImportDialogueSettings
 	{
 	}
 
-	/** Path to the folder to import the audio from. This folder is expected to contain culture sub-folders, which in turn contain the raw WAV files to import. */
+	/** Path to the folder to import the audio from. This folder is expected to contain language sub-folders, which in turn contain the raw WAV files to import. */
 	UPROPERTY(config, EditAnywhere, Category="Dialogue")
 	FDirectoryPath RawAudioPath;
 
-	/** Folder in which to create the generated sound waves. This is relative to the root of the L10N culture folder (or the root content folder if importing native dialogue as source dialogue). */
+	/** Folder in which to create the generated sound waves. This is relative to the root of the L10N language folder (or the root content folder if importing native dialogue as source dialogue). */
 	UPROPERTY(config, EditAnywhere, Category="Dialogue")
 	FString ImportedDialogueFolder;
 
-	/** Should the dialogue for the native culture be imported as if it were source audio? If false, the native culture dialogue will be imported as localized data for the native culture. */
+	/** Should the dialogue for the native language be imported as if it were source audio? If false, the native language dialogue will be imported as localized data for the native language. */
 	UPROPERTY(config, EditAnywhere, Category="Dialogue")
 	bool bImportNativeAsSource;
 };
@@ -400,11 +400,11 @@ struct FCultureStatistics
 	}
 
 
-	/* The ISO name for this culture. */
-	UPROPERTY(config, EditAnywhere, Category = "Culture")
+	/* The ISO name for this language. */
+	UPROPERTY(config, EditAnywhere, Category = "Language", DisplayName="Language")
 	FString CultureName;
 
-	/* The estimated number of words that have been localized for this culture. */
+	/* The estimated number of words that have been localized for this language. */
 	UPROPERTY(Transient, EditAnywhere, Category = "Statistics")
 	uint32 WordCount;
 };
@@ -485,12 +485,12 @@ struct FLocalizationTargetSettings
 	UPROPERTY(config, EditAnywhere, Category = "Import Dialogue", meta=(ShowOnlyInnerProperties))
 	FLocalizationImportDialogueSettings ImportDialogueSettings;
 
-	/* The index of the native culture among the supported cultures. */
-	UPROPERTY(config, EditAnywhere, Category = "Cultures")
+	/* The index of the native language among the supported languages. */
+	UPROPERTY(config, EditAnywhere, Category = "Languages")
 	int32 NativeCultureIndex;
 
-	/* Cultures for which the source text is being localized for.*/
-	UPROPERTY(config, EditAnywhere, Category = "Cultures")
+	/* Language for which the source text is being localized for.*/
+	UPROPERTY(config, EditAnywhere, Category = "Languages")
 	TArray<FCultureStatistics> SupportedCulturesStatistics;
 };
 
diff --git a/Engine/Source/Editor/LocalizationDashboard/Private/LocalizationTargetDetailCustomization.cpp b/Engine/Source/Editor/LocalizationDashboard/Private/LocalizationTargetDetailCustomization.cpp
--- a/Engine/Source/Editor/LocalizationDashboard/Private/LocalizationTargetDetailCustomization.cpp
+++ b/Engine/Source/Editor/LocalizationDashboard/Private/LocalizationTargetDetailCustomization.cpp
@@ -154,14 +154,14 @@ public:
 
 void FLocalizationTargetEditorCommands::RegisterCommands() 
 {
-	UI_COMMAND(GatherText, "Gather Text", "Gather text for all cultures of this target.", EUserInterfaceActionType::Button, FInputChord());
-	UI_COMMAND(ImportTextAllCultures, "Import Text", "Import translations for all cultures of this target.", EUserInterfaceActionType::Button, FInputChord());
-	UI_COMMAND(ExportTextAllCultures, "Export Text", "Export translations for all cultures of this target.", EUserInterfaceActionType::Button, FInputChord());
-	UI_COMMAND(ImportDialogueScriptAllCultures, "Import Script", "Import dialogue scripts for all cultures of this target.", EUserInterfaceActionType::Button, FInputChord());
-	UI_COMMAND(ExportDialogueScriptAllCultures, "Export Script", "Export dialogue scripts for all cultures of this target.", EUserInterfaceActionType::Button, FInputChord());
-	UI_COMMAND(ImportDialogueAllCultures, "Import Dialogue", "Import dialogue WAV files for all cultures of this target.", EUserInterfaceActionType::Button, FInputChord());
-	UI_COMMAND(CountWords, "Count Words", "Count translations for all cultures of this target.", EUserInterfaceActionType::Button, FInputChord());
-	UI_COMMAND(CompileTextAllCultures, "Compile Text", "Compile translations for all cultures of this target.", EUserInterfaceActionType::Button, FInputChord());
+	UI_COMMAND(GatherText, "Gather Text", "Gather text for all languages of this target.", EUserInterfaceActionType::Button, FInputChord());
+	UI_COMMAND(ImportTextAllCultures, "Import Text", "Import translations for all languages of this target.", EUserInterfaceActionType::Button, FInputChord());
+	UI_COMMAND(ExportTextAllCultures, "Export Text", "Export translations for all languages of this target.", EUserInterfaceActionType::Button, FInputChord());
+	UI_COMMAND(ImportDialogueScriptAllCultures, "Import Script", "Import dialogue scripts for all languages of this target.", EUserInterfaceActionType::Button, FInputChord());
+	UI_COMMAND(ExportDialogueScriptAllCultures, "Export Script", "Export dialogue scripts for all languages of this target.", EUserInterfaceActionType::Button, FInputChord());
+	UI_COMMAND(ImportDialogueAllCultures, "Import Dialogue", "Import dialogue WAV files for all languages of this target.", EUserInterfaceActionType::Button, FInputChord());
+	UI_COMMAND(CountWords, "Count Words", "Count translations for all languages of this target.", EUserInterfaceActionType::Button, FInputChord());
+	UI_COMMAND(CompileTextAllCultures, "Compile Text", "Compile translations for all languages of this target.", EUserInterfaceActionType::Button, FInputChord());
 }
 
 void FLocalizationTargetDetailCustomization::CustomizeDetails(IDetailLayoutBuilder& DetailBuilder)
@@ -288,7 +288,7 @@ void FLocalizationTargetDetailCustomization::CustomizeDetails(IDetailLayoutBuild
 
 		TAttribute<FText> GatherToolTipTextAttribute = TAttribute<FText>::Create(TAttribute<FText>::FGetter::CreateLambda([this]() -> FText
 		{
-			return CanGatherText() ? FLocalizationTargetEditorCommands::Get().GatherText->GetDescription() : LOCTEXT("GatherDisabledToolTip", "Must have a native culture specified in order to gather.");
+			return CanGatherText() ? FLocalizationTargetEditorCommands::Get().GatherText->GetDescription() : LOCTEXT("GatherDisabledToolTip", "Must have a native language specified in order to gather.");
 		}));
 		CommandList->MapAction(FLocalizationTargetEditorCommands::Get().GatherText, FExecuteAction::CreateSP(this, &FLocalizationTargetDetailCustomization::GatherText), FCanExecuteAction::CreateSP(this, &FLocalizationTargetDetailCustomization::CanGatherText));
 		ToolBarBuilder.AddToolBarButton(FLocalizationTargetEditorCommands::Get().GatherText, NAME_None, TAttribute<FText>(), GatherToolTipTextAttribute, FSlateIcon(FAppStyle::GetAppStyleSetName(), "LocalizationTargetEditor.GatherText"));
@@ -348,7 +348,7 @@ void FLocalizationTargetDetailCustomization::CustomizeDetails(IDetailLayoutBuild
 						.VAlignCell(VAlign_Center)
 						.FillWidth(0.1f)
 						+SHeaderRow::Column("Culture")
-						.DefaultLabel( NSLOCTEXT("LocalizationCulture", "CultureColumnLabel", "Culture"))
+						.DefaultLabel( NSLOCTEXT("LocalizationCulture", "CultureColumnLabel", "Language"))
 						.HAlignHeader(HAlign_Fill)
 						.HAlignCell(HAlign_Fill)
 						.VAlignCell(VAlign_Center)
@@ -381,7 +381,7 @@ void FLocalizationTargetDetailCustomization::CustomizeDetails(IDetailLayoutBuild
 						.ButtonContent()
 						[
 							SNew(STextBlock)
-							.Text(NSLOCTEXT("LocalizationCulture", "AddNewCultureButtonLabel", "Add New Culture"))
+							.Text(NSLOCTEXT("LocalizationCulture", "AddNewCultureButtonLabel", "Add New Language"))
 						]
 						.MenuContent()
 							[
@@ -1211,7 +1211,7 @@ void FLocalizationTargetDetailCustomization::BuildListedCulturesList()
 		}
 		else
 		{
-			NoSupportedCulturesErrorText->SetError(LOCTEXT("NoSupportedCulturesError", "At least one supported culture must be specified."));
+			NoSupportedCulturesErrorText->SetError(LOCTEXT("NoSupportedCulturesError", "At least one supported language must be specified."));
 		}
 	}
 }
