Fixed crash when re-ordering localization gather paths in the Localization Dashboard

diff --git a/Engine/Source/Editor/LocalizationDashboard/Private/GatherTextDetailCustomizations.cpp b/Engine/Source/Editor/LocalizationDashboard/Private/GatherTextDetailCustomizations.cpp
--- a/Engine/Source/Editor/LocalizationDashboard/Private/GatherTextDetailCustomizations.cpp
+++ b/Engine/Source/Editor/LocalizationDashboard/Private/GatherTextDetailCustomizations.cpp
@@ -261,7 +261,9 @@ namespace
 
 		TArray<const void*> RawData;
 		PathRootPropertyHandle->AccessRawData(RawData);
-		return *(ELocalizationGatherPathRoot*)RawData.Top();
+		return RawData.Top()
+			? *(ELocalizationGatherPathRoot*)RawData.Top()
+			: ELocalizationGatherPathRoot::Auto;
 	}
 
 	FString SGatherTextPathPicker::GetResolvedPathRoot() const
@@ -288,17 +290,18 @@ void FGatherTextSearchDirectoryStructCustomization::CustomizeStructHeader( TShar
 		{
 			TArray<const void*> RawData;
 			StructPropertyHandle->AccessRawData(RawData);
-			const FGatherTextSearchDirectory& Struct = *(reinterpret_cast<const FGatherTextSearchDirectory*>(RawData.Top()));
-
-			TArray<UObject*> OuterObjects;
-			StructPropertyHandle->GetOuterObjects(OuterObjects);
-			ULocalizationTarget* const LocalizationTarget = Cast<ULocalizationTarget>(OuterObjects.Top());
-			if (LocalizationTarget)
+			if (const FGatherTextSearchDirectory* Struct = (reinterpret_cast<const FGatherTextSearchDirectory*>(RawData.Top())))
 			{
-				FText Error;
-				if (!Struct.Validate(LocalizationTarget->IsMemberOfEngineTargetSet(), Error))
+				TArray<UObject*> OuterObjects;
+				StructPropertyHandle->GetOuterObjects(OuterObjects);
+				ULocalizationTarget* const LocalizationTarget = Cast<ULocalizationTarget>(OuterObjects.Top());
+				if (LocalizationTarget)
 				{
-					return Error;
+					FText Error;
+					if (!Struct->Validate(LocalizationTarget->IsMemberOfEngineTargetSet(), Error))
+					{
+						return Error;
+					}
 				}
 			}
 		}
